make the BP editor library an engine module instead of a plugin

Moving files from the plugins dir to the editor modules dir

#rnx

diff --git a/Engine/Plugins/BlueprintEditorLibrary/BlueprintEditorLibrary.uplugin b/Engine/Plugins/BlueprintEditorLibrary/BlueprintEditorLibrary.uplugin
deleted file mode 100644
diff --git a/Engine/Plugins/BlueprintEditorLibrary/Source/BlueprintEditorLibrary/BlueprintEditorLibrary.Build.cs b/Engine/Plugins/BlueprintEditorLibrary/Source/BlueprintEditorLibrary/BlueprintEditorLibrary.Build.cs
deleted file mode 100644
diff --git a/Engine/Plugins/BlueprintEditorLibrary/Source/BlueprintEditorLibrary/Private/BlueprintEditorLibrary.cpp b/Engine/Plugins/BlueprintEditorLibrary/Source/BlueprintEditorLibrary/Private/BlueprintEditorLibrary.cpp
deleted file mode 100644
diff --git a/Engine/Plugins/BlueprintEditorLibrary/Source/BlueprintEditorLibrary/Private/BlueprintEditorLibraryModule.cpp b/Engine/Plugins/BlueprintEditorLibrary/Source/BlueprintEditorLibrary/Private/BlueprintEditorLibraryModule.cpp
deleted file mode 100644
diff --git a/Engine/Plugins/BlueprintEditorLibrary/Source/BlueprintEditorLibrary/Public/BlueprintEditorLibrary.h b/Engine/Plugins/BlueprintEditorLibrary/Source/BlueprintEditorLibrary/Public/BlueprintEditorLibrary.h
deleted file mode 100644
diff --git a/Engine/Plugins/BlueprintEditorLibrary/Source/BlueprintEditorLibrary/Public/BlueprintEditorLibraryModule.h b/Engine/Plugins/BlueprintEditorLibrary/Source/BlueprintEditorLibrary/Public/BlueprintEditorLibraryModule.h
deleted file mode 100644
diff --git a/Engine/Source/Editor/BlueprintEditorLibrary/BlueprintEditorLibrary.Build.cs b/Engine/Source/Editor/BlueprintEditorLibrary/BlueprintEditorLibrary.Build.cs
new file mode 100644
--- /dev/null
+++ b/Engine/Source/Editor/BlueprintEditorLibrary/BlueprintEditorLibrary.Build.cs
@@ -0,0 +1,40 @@
+// Copyright Epic Games, Inc. All Rights Reserved.
+
+using UnrealBuildTool;
+using System.IO;
+
+public class BlueprintEditorLibrary : ModuleRules
+{
+	public BlueprintEditorLibrary(ReadOnlyTargetRules Target) : base(Target)
+	{
+		PCHUsage = ModuleRules.PCHUsageMode.UseExplicitOrSharedPCHs;
+		PublicIncludePaths.Add(Path.Combine(ModuleDirectory, "Public"));
+
+		PrivateIncludePaths.AddRange(
+			new string[] 
+			{
+				Path.Combine(ModuleDirectory, "Private"),
+				//"Editor/BlueprintEditorLibrary/Private",
+			}
+		);
+			
+		
+		PublicDependencyModuleNames.AddRange(
+			new string[]
+			{
+				"Core",
+				"BlueprintGraph"
+			}
+		);
+			
+		PrivateDependencyModuleNames.AddRange(
+			new string[]
+			{
+				"CoreUObject",
+				"Engine",
+				"UnrealEd",
+				"AnimGraph"
+			}
+		);
+	}
+}
diff --git a/Engine/Source/Editor/BlueprintEditorLibrary/Private/BlueprintEditorLibrary.cpp b/Engine/Source/Editor/BlueprintEditorLibrary/Private/BlueprintEditorLibrary.cpp
new file mode 100644
--- /dev/null
+++ b/Engine/Source/Editor/BlueprintEditorLibrary/Private/BlueprintEditorLibrary.cpp
@@ -0,0 +1,521 @@
+// Copyright Epic Games, Inc. All Rights Reserved.
+
+#include "BlueprintEditorLibrary.h"
+#include "BlueprintEditorLibraryModule.h"
+#include "Kismet2/BlueprintEditorUtils.h"
+#include "EdGraph/EdGraph.h"
+#include "EdGraph/EdGraphNode.h"
+#include "EdGraphSchema_K2.h"
+#include "EdGraphNode_Comment.h"
+#include "EdGraphUtilities.h"
+#include "BlueprintTypePromotion.h"
+#include "K2Node_PromotableOperator.h"
+#include "K2Node_CommutativeAssociativeBinaryOperator.h"
+#include "Kismet2/KismetEditorUtilities.h"
+#include "BlueprintNodeSpawner.h"
+#include "EdGraphUtilities.h"
+#include "AnimGraphNode_Base.h"
+#include "Components/TimelineComponent.h"
+
+#define LOCTEXT_NAMESPACE "BlueprintEditorLibrary"
+
+DEFINE_LOG_CATEGORY(LogBlueprintEditorLib);
+
+///////////////////////////////////////////////////////////
+// InternalBlueprintEditorLibrary
+
+namespace InternalBlueprintEditorLibrary
+{
+	/**
+	* Replace the OldNode with the NewNode and reconnect it's pins. If the pins don't
+	* exist on the NewNode, then orphan the connections.
+	*
+	* @param OldNode		The old node to replace
+	* @param NewNode		The new node to put in the old node's place
+	*/
+	static void ReplaceOldNodeWithNew(UEdGraphNode* OldNode, UEdGraphNode* NewNode)
+	{
+		const UEdGraphSchema_K2* Schema = GetDefault<UEdGraphSchema_K2>();
+
+		if (Schema && OldNode && NewNode)
+		{
+			TMap<FName, FName> OldToNewPinMap;
+			for (UEdGraphPin* Pin : OldNode->Pins)
+			{
+				if (Pin->ParentPin != nullptr)
+				{
+					// ReplaceOldNodeWithNew() will take care of mapping split pins (as long as the parents are properly mapped)
+					continue;
+				}
+				else if (Pin->PinName == UEdGraphSchema_K2::PN_Self)
+				{
+					// there's no analogous pin, signal that we're expecting this
+					OldToNewPinMap.Add(Pin->PinName, NAME_None);
+				}
+				else
+				{
+					// The input pins follow the same naming scheme
+					OldToNewPinMap.Add(Pin->PinName, Pin->PinName);
+				}
+			}
+			
+			Schema->ReplaceOldNodeWithNew(OldNode, NewNode, OldToNewPinMap);
+			// reconstructing the node will clean up any
+			// incorrect default values that may have been copied over
+			NewNode->ReconstructNode();
+		}
+	}
+
+	/**
+	* Returns true if any of these nodes pins have any links. Does not check for 
+	* a default value on pins
+	*
+	* @param Node		The node to check
+	*
+	* @return bool		True if the node has any links, false otherwise.
+	*/
+	static bool NodeHasAnyConnections(const UEdGraphNode* Node)
+	{
+		if (Node)
+		{
+			for (const UEdGraphPin* Pin : Node->Pins)
+			{
+				if (Pin && Pin->LinkedTo.Num() > 0)
+				{
+					return true;
+				}
+			}
+		}
+
+		return false;
+	}
+};
+
+///////////////////////////////////////////////////////////
+// UBlueprintEditorLibrary
+
+UBlueprintEditorLibrary::UBlueprintEditorLibrary(const FObjectInitializer& ObjectInitializer)
+: Super(ObjectInitializer)
+{
+}
+
+void UBlueprintEditorLibrary::ReplaceVariableReferences(UBlueprint* Blueprint, const FName OldVarName, const FName NewVarName)
+{
+	if (!Blueprint || OldVarName.IsNone() || NewVarName.IsNone())
+	{
+		return;
+	}
+
+	FBlueprintEditorUtils::RenameVariableReferences(Blueprint, Blueprint->GeneratedClass, OldVarName, NewVarName);
+}
+
+void UBlueprintEditorLibrary::ReplaceK2Nodes(UBlueprint* Blueprint, TSubclassOf<UEdGraphNode> OldNodeType, TSubclassOf<UEdGraphNode> NewNodeType)
+{
+	if (!Blueprint)
+	{
+		return;
+	}
+
+	TArray<UEdGraph*> AllGraphs;
+	Blueprint->GetAllGraphs(AllGraphs);
+
+	for (UEdGraph* Graph : AllGraphs)
+	{
+		Graph->Modify();
+
+		// Iterate backwards through the nodes because we will be removing some
+		for (int32 i = Graph->Nodes.Num() - 1; i >= 0; i--)
+		{
+			UEdGraphNode* OriginalNode = Graph->Nodes[i];
+			if (OriginalNode && OriginalNode->IsA(OldNodeType))
+			{
+				// Spawn a new node of the new class
+				UEdGraphNode* NewNode = NewObject<UEdGraphNode>(Graph, NewNodeType);
+				NewNode->CreateNewGuid();
+				NewNode->PostPlacedNewNode();
+				NewNode->AllocateDefaultPins();
+				NewNode->SetFlags(RF_Transactional);
+
+				NewNode->NodePosX = OriginalNode->NodePosX;
+				NewNode->NodePosY = OriginalNode->NodePosY;
+
+				InternalBlueprintEditorLibrary::ReplaceOldNodeWithNew(OriginalNode, NewNode);
+			}
+		}
+	}
+}
+
+UEdGraph* UBlueprintEditorLibrary::FindEventGraph(UBlueprint* Blueprint)
+{
+	return Blueprint ? FBlueprintEditorUtils::FindEventGraph(Blueprint) : nullptr;
+}
+
+UEdGraph* UBlueprintEditorLibrary::FindGraph(UBlueprint* Blueprint, FName GraphName)
+{
+	if (Blueprint && !GraphName.IsNone())
+	{
+		TArray<UEdGraph*> AllGraphs;
+		Blueprint->GetAllGraphs(AllGraphs);
+
+		for (UEdGraph* CurrentGraph : AllGraphs)
+		{
+			if (CurrentGraph->GetFName() == GraphName)
+			{
+				return CurrentGraph;
+			}
+		}
+	}
+
+	return nullptr;
+}
+
+UK2Node_PromotableOperator* CreateOpNode(const FName OpName, UEdGraph* Graph, const int32 AdditionalPins)
+{
+	if (!Graph)
+	{
+		return nullptr;
+	}
+
+	// The spawner will be null if type promo isn't enabled
+	if (UBlueprintFunctionNodeSpawner* Spawner = FTypePromotion::GetOperatorSpawner(OpName))
+	{
+		// Spawn a new node!
+		IBlueprintNodeBinder::FBindingSet Bindings;
+		FVector2D SpawnLoc{};
+		UK2Node_PromotableOperator* NewOpNode = Cast<UK2Node_PromotableOperator>(Spawner->Invoke(Graph, Bindings, SpawnLoc));
+		check(NewOpNode);
+
+		// Add the necessary number of additional pins
+		for (int32 i = 0; i < AdditionalPins; ++i)
+		{
+			NewOpNode->AddInputPin();
+		}
+
+		return NewOpNode;
+	}
+
+	return nullptr;
+}
+
+void UBlueprintEditorLibrary::UpgradeOperatorNodes(UBlueprint* Blueprint)
+{
+	if (!Blueprint)
+	{
+		return;
+	}
+
+	if (TypePromoDebug::IsTypePromoEnabled())
+	{
+		// Ensure that we have promotable operator node spanners available to us. 
+		// They will be empty if the editor hasn't been opened or 
+		if (FBlueprintActionDatabase* Actions = FBlueprintActionDatabase::TryGet())
+		{
+			Actions->RefreshAll();
+		}
+	}
+	else
+	{
+		UE_LOG(LogBlueprintEditorLib, Warning, TEXT("Type Promotion is not enabled! Cannot upgrade operator nodes. Set 'BP.TypePromo.IsEnabled' to true and try again."));
+		return;
+	}
+
+	TArray<UEdGraph*> AllGraphs;
+	Blueprint->GetAllGraphs(AllGraphs);
+	Blueprint->Modify();
+
+	for (UEdGraph* Graph : AllGraphs)
+	{
+		check(Graph);
+
+		Graph->Modify();
+
+		for (int32 i = Graph->Nodes.Num() - 1; i >= 0; --i)
+		{
+			// Not every function that we want to upgrade is a CommunicativeBinaryOpNode
+			// Some are just regular CallFunction nodes; Vector + Float is an example of this
+			if (UK2Node_CallFunction* OldOpNode = Cast<UK2Node_CallFunction>(Graph->Nodes[i]))
+			{
+				UFunction* Func = OldOpNode->GetTargetFunction();
+				UEdGraph* OwningGraph = OldOpNode->GetGraph();
+
+				// We should only be modifying nodes within the graph that we want
+				ensure(OwningGraph == Graph);
+				
+				// Don't bother with non-promotable functions or things that are already promotable operators
+				if (!FTypePromotion::IsFunctionPromotionReady(Func) || OldOpNode->IsA<UK2Node_PromotableOperator>())
+				{
+					continue;
+				}
+
+				FName OpName = FTypePromotion::GetOpNameFromFunction(Func);
+
+				UK2Node_CommutativeAssociativeBinaryOperator* BinaryOpNode = Cast<UK2Node_CommutativeAssociativeBinaryOperator>(OldOpNode);
+
+				// Spawn a new node!
+				UK2Node_PromotableOperator* NewOpNode = CreateOpNode(
+					OpName,
+					OwningGraph,
+					BinaryOpNode ? BinaryOpNode->GetNumberOfAdditionalInputs() : 0
+				);
+
+				// If there is a node that is a communicative op node but is not promotable
+				// then the node will be null
+				if (!NewOpNode)
+				{
+					UE_LOG(LogBlueprintEditorLib, Warning, TEXT("Failed to spawn new operator node!"));
+					continue;
+				}
+
+				NewOpNode->NodePosX = OldOpNode->NodePosX;
+				NewOpNode->NodePosY = OldOpNode->NodePosY;
+
+				InternalBlueprintEditorLibrary::ReplaceOldNodeWithNew(OldOpNode, NewOpNode);
+			}
+		}
+	}
+}
+
+void UBlueprintEditorLibrary::CompileBlueprint(UBlueprint* Blueprint)
+{
+	if (Blueprint)
+	{
+		// Skip saving this to avoid possible tautologies when saving and allow the user to manually save
+		EBlueprintCompileOptions Flags = EBlueprintCompileOptions::SkipSave;
+		FKismetEditorUtilities::CompileBlueprint(Blueprint, Flags);
+	}
+}
+
+UEdGraph* UBlueprintEditorLibrary::AddFunctionGraph(UBlueprint* Blueprint, const FString& FuncName)
+{
+	if (!Blueprint)
+	{
+		return nullptr;
+	}
+
+	Blueprint->Modify();
+	UEdGraph* NewGraph = FBlueprintEditorUtils::CreateNewGraph(
+		Blueprint, 
+		FBlueprintEditorUtils::FindUniqueKismetName(Blueprint, FuncName),
+		UEdGraph::StaticClass(), 
+		UEdGraphSchema_K2::StaticClass()
+	);
+
+	FBlueprintEditorUtils::AddFunctionGraph<UFunction>(Blueprint, NewGraph, /* bIsUserCreated = */ true, /* SignatureFromObject = */ nullptr);
+
+	return NewGraph;
+}
+
+void UBlueprintEditorLibrary::RemoveFunctionGraph(UBlueprint* Blueprint, FName FuncName)
+{
+	if (!Blueprint)
+	{
+		return;
+	}
+
+	// Find the function graph of this name
+	UEdGraph* FunctionGraph = nullptr;
+	for (UEdGraph* Graph : Blueprint->FunctionGraphs)
+	{
+		if (Graph->GetFName() == FuncName)
+		{
+			FunctionGraph = Graph;
+			break;
+		}
+	}
+
+	// Remove the function graph if we can
+	if (FunctionGraph && FunctionGraph->bAllowDeletion)
+	{
+		Blueprint->Modify();
+		FBlueprintEditorUtils::RemoveGraph(Blueprint, FunctionGraph, EGraphRemoveFlags::MarkTransient);
+	}
+	else
+	{
+		UE_LOG(LogBlueprintEditorLib, Warning, TEXT("Failed to remove function '%s' on blueprint '%s'!"), *FuncName.ToString(), *Blueprint->GetFriendlyName());
+	}
+}
+
+void UBlueprintEditorLibrary::RemoveUnusedNodes(UBlueprint* Blueprint)
+{
+	if (!Blueprint)
+	{
+		return;
+	}
+
+	TArray<UEdGraph*> AllGraphs;
+	Blueprint->GetAllGraphs(AllGraphs);
+	Blueprint->Modify();
+
+	for (UEdGraph* Graph : AllGraphs)
+	{
+		// Skip non-editable graphs
+		if (!Graph || FBlueprintEditorUtils::IsGraphReadOnly(Graph))
+		{
+			continue;
+		}
+
+		Graph->Modify();
+		int32 NumNodesRemoved = 0;
+
+		for (int32 i = Graph->Nodes.Num() - 1; i >= 0; --i)
+		{
+			UEdGraphNode* Node = Graph->Nodes[i];
+
+			// We only want to delete user facing nodes because this is meant 
+			// to be a BP refactoring/cleanup tool. Anim graph nodes can still 
+			// be valid with no pin connections made to them
+			if (Node->CanUserDeleteNode() && 
+				!Node->IsA<UAnimGraphNode_Base>() && 
+				!InternalBlueprintEditorLibrary::NodeHasAnyConnections(Node))
+			{
+				Node->BreakAllNodeLinks();
+				Graph->RemoveNode(Node);
+				++NumNodesRemoved;
+			}
+		}
+
+		// Notify a change to the graph if nodes have been removed
+		if (NumNodesRemoved > 0)
+		{
+			Graph->NotifyGraphChanged();
+		}
+	}
+
+	FBlueprintEditorUtils::MarkBlueprintAsStructurallyModified(Blueprint);
+}
+
+void UBlueprintEditorLibrary::RemoveGraph(UBlueprint* Blueprint, UEdGraph* Graph)
+{
+	if (!Blueprint || !Graph)
+	{
+		return;
+	}
+
+	FBlueprintEditorUtils::RemoveGraph(Blueprint, Graph, EGraphRemoveFlags::MarkTransient);
+}
+
+void UBlueprintEditorLibrary::RenameGraph(UEdGraph* Graph, const FString& NewNameStr)
+{
+	if (!Graph)
+	{
+		return;
+	}
+		
+	FBlueprintEditorUtils::RenameGraph(Graph, NewNameStr);
+}
+
+UBlueprint* UBlueprintEditorLibrary::GetBlueprintAsset(UObject* Object)
+{
+	return Cast<UBlueprint>(Object);
+}
+
+void UBlueprintEditorLibrary::ReparentBlueprint(UBlueprint* Blueprint, UClass* NewParentClass)
+{
+	if (!Blueprint || !NewParentClass)
+	{
+		UE_LOG(LogBlueprintEditorLib, Warning, TEXT("Failed to reparent blueprint!"));
+		return;
+	}
+
+	if (NewParentClass == Blueprint->ParentClass)
+	{
+		UE_LOG(LogBlueprintEditorLib, Warning, TEXT("'%s' is already parented to class '%s'!"), *Blueprint->GetFriendlyName(), *NewParentClass->GetName());
+		return;
+	}
+
+	// There could be possible data loss if reparenting outside the current class hierarchy
+	if (!Blueprint->ParentClass || !NewParentClass->GetDefaultObject()->IsA(Blueprint->ParentClass))
+	{
+		UE_LOG(LogBlueprintEditorLib, Warning, TEXT("'%s' class heirarcy is changing, there could be possible data loss!"), *Blueprint->GetFriendlyName());
+	}
+
+	UClass* OriginalParentClass = Blueprint->ParentClass;
+	Blueprint->ParentClass = NewParentClass;
+
+	FBlueprintEditorUtils::RefreshAllNodes(Blueprint);
+	FBlueprintEditorUtils::MarkBlueprintAsModified(Blueprint);
+
+	CompileBlueprint(Blueprint);
+}
+
+bool UBlueprintEditorLibrary::GatherUnusedVariables(const UBlueprint* Blueprint, TArray<FProperty*>& OutProperties)
+{
+	if (!Blueprint)
+	{
+		return false;
+	}
+
+	bool bHasAtLeastOneVariableToCheck = false;
+
+	for (TFieldIterator<FProperty> PropertyIt(Blueprint->SkeletonGeneratedClass, EFieldIteratorFlags::ExcludeSuper); PropertyIt; ++PropertyIt)
+	{
+		FProperty* Property = *PropertyIt;
+		// Don't show delegate properties, there is special handling for these
+		const bool bDelegateProp = Property->IsA(FDelegateProperty::StaticClass()) || Property->IsA(FMulticastDelegateProperty::StaticClass());
+		const bool bShouldShowProp = (!Property->HasAnyPropertyFlags(CPF_Parm) && Property->HasAllPropertyFlags(CPF_BlueprintVisible) && !bDelegateProp);
+
+		if (bShouldShowProp)
+		{
+			bHasAtLeastOneVariableToCheck = true;
+			FName VarName = Property->GetFName();
+
+			const int32 VarInfoIndex = FBlueprintEditorUtils::FindNewVariableIndex(Blueprint, VarName);
+			const bool bHasVarInfo = (VarInfoIndex != INDEX_NONE);
+
+			const FObjectPropertyBase* ObjectProperty = CastField<const FObjectPropertyBase>(Property);
+			bool bIsTimeline = ObjectProperty &&
+				ObjectProperty->PropertyClass &&
+				ObjectProperty->PropertyClass->IsChildOf(UTimelineComponent::StaticClass());
+			if (!bIsTimeline && bHasVarInfo && !FBlueprintEditorUtils::IsVariableUsed(Blueprint, VarName))
+			{
+				OutProperties.Add(Property);
+			}
+		}
+	}
+
+	return bHasAtLeastOneVariableToCheck;
+}
+
+int32 UBlueprintEditorLibrary::RemoveUnusedVariables(UBlueprint* Blueprint)
+{
+	if (!Blueprint)
+	{
+		return 0;
+	}
+
+	// Gather FProperties from this BP and see if we can remove any
+	TArray<FProperty*> VariableProperties;
+	UBlueprintEditorLibrary::GatherUnusedVariables(Blueprint, VariableProperties);
+	
+	// No variables can be removed from this blueprint
+	if (VariableProperties.Num() == 0)
+	{
+		return 0;
+	}
+
+	// Get the variables by name so that we can bulk remove them and print them out to the log
+	TArray<FName> VariableNames;
+	FString PropertyList;
+	VariableNames.Reserve(VariableProperties.Num());
+	for (int32 Index = 0; Index < VariableProperties.Num(); ++Index)
+	{
+		VariableNames.Add(VariableProperties[Index]->GetFName());
+		if (PropertyList.IsEmpty())
+		{
+			PropertyList = UEditorEngine::GetFriendlyName(VariableProperties[Index]);
+		}
+		else
+		{
+			PropertyList += FString::Printf(TEXT(", %s"), *UEditorEngine::GetFriendlyName(VariableProperties[Index]));
+		}
+	}
+
+	const int32 NumRemovedVars = VariableNames.Num();
+	// Remove the variables by name
+	FBlueprintEditorUtils::BulkRemoveMemberVariables(Blueprint, VariableNames);
+
+	UE_LOG(LogBlueprintEditorLib, Log, TEXT("The following variable(s) were deleted successfully: %s."), *PropertyList);
+	return NumRemovedVars;
+}
+
+#undef LOCTEXT_NAMESPACE	// "BlueprintEditorLibrary"
\ No newline at end of file
diff --git a/Engine/Source/Editor/BlueprintEditorLibrary/Private/BlueprintEditorLibraryModule.cpp b/Engine/Source/Editor/BlueprintEditorLibrary/Private/BlueprintEditorLibraryModule.cpp
new file mode 100644
--- /dev/null
+++ b/Engine/Source/Editor/BlueprintEditorLibrary/Private/BlueprintEditorLibraryModule.cpp
@@ -0,0 +1,22 @@
+// Copyright Epic Games, Inc. All Rights Reserved.
+
+#include "BlueprintEditorLibraryModule.h"
+
+#define LOCTEXT_NAMESPACE "FBlueprintEditorLibraryModule"
+
+void FBlueprintEditorLibraryModule::StartupModule()
+{
+	// This code will execute after your module is loaded into memory; the exact timing is specified in the .uplugin file per-module
+	UE_LOG(LogTemp, Warning, TEXT("BP Editor Lib Loaded"));
+}
+
+void FBlueprintEditorLibraryModule::ShutdownModule()
+{
+	// This function may be called during shutdown to clean up your module.  For modules that support dynamic reloading,
+	// we call this function before unloading the module.	
+	UE_LOG(LogTemp, Warning, TEXT("BP Editor Lib Loaded"));
+}
+
+#undef LOCTEXT_NAMESPACE
+
+IMPLEMENT_MODULE(FBlueprintEditorLibraryModule, BlueprintEditorLibrary)
\ No newline at end of file
diff --git a/Engine/Source/Editor/BlueprintEditorLibrary/Public/BlueprintEditorLibrary.h b/Engine/Source/Editor/BlueprintEditorLibrary/Public/BlueprintEditorLibrary.h
new file mode 100644
--- /dev/null
+++ b/Engine/Source/Editor/BlueprintEditorLibrary/Public/BlueprintEditorLibrary.h
@@ -0,0 +1,166 @@
+// Copyright Epic Games, Inc. All Rights Reserved.
+
+#pragma once
+
+#include "CoreMinimal.h"
+#include "UObject/ObjectMacros.h"
+#include "Kismet/BlueprintFunctionLibrary.h"
+#include "Components/ActorComponent.h"
+#include "EdGraph/EdGraphNode.h"
+#include "BlueprintEditorLibrary.generated.h"
+
+DECLARE_LOG_CATEGORY_EXTERN(LogBlueprintEditorLib, Warning, All);
+
+UCLASS()
+class BLUEPRINTEDITORLIBRARY_API UBlueprintEditorLibrary : public UBlueprintFunctionLibrary
+{
+	GENERATED_UCLASS_BODY()
+public:
+
+	/**
+	* Replace any references of variables with the OldVarName to references of those with the NewVarName if possible
+	*
+	* @param Blueprint		Blueprint to replace the variable references on
+	* @param OldVarName		The variable you want replaced
+	* @param NewVarName		The new variable that will be used in the old one's place
+	*/
+	UFUNCTION(BlueprintCallable, Category = "Blueprint Upgrade Tools")
+	static void ReplaceVariableReferences(UBlueprint* Blueprint, const FName OldVarName, const FName NewVarName);
+
+	/**
+	* Replace any existing k2node class with the new one, attempting to remap pins as we go
+	*
+	* @param Blueprint		Blueprint to replace node types on
+	* @param OldNodeType	Old node type that will be removed
+	* @param NewNodeType	New node type that will be put in the old node's place
+	*/
+	UFUNCTION(BlueprintCallable, Category = "Blueprint Upgrade Tools")
+	static void ReplaceK2Nodes(UBlueprint* Blueprint, TSubclassOf<UEdGraphNode> OldNodeType, TSubclassOf<UEdGraphNode> NewNodeType);
+
+	/**
+	* Finds the event graph of the given blueprint. Null if it doesn't have one. This will only return
+	* the primary event graph of the blueprint (the graph named "EventGraph").
+	*
+	* @param Blueprint		Blueprint to search for the event graph on
+	*
+	* @return UEdGraph*		Event graph of the blueprint if it has one, null if it doesn't have one
+	*/
+	UFUNCTION(BlueprintCallable, Category = "Blueprint Upgrade Tools")
+	static UEdGraph* FindEventGraph(UBlueprint* Blueprint);
+
+	/**
+	* Finds the graph with the given name on the blueprint. Null if it doesn't have one. 
+	*
+	* @param Blueprint		Blueprint to search
+	* @param GraphName		The name of the graph to search for 
+	*
+	* @return UEdGraph*		Pointer to the graph with the given name, null if not found
+	*/
+	UFUNCTION(BlueprintCallable, Category = "Blueprint Upgrade Tools")
+	static UEdGraph* FindGraph(UBlueprint* Blueprint, FName GraphName);
+
+	/**
+	* Replace any old operator nodes (float + float, vector + float, int + vector, etc)
+	* with the newer Promotable Operator version of the node. Preserve any connections the
+	* original node had to the newer version of the node. 
+	*
+	* @param Blueprint	Blueprint to upgrade
+	*/
+	UFUNCTION(BlueprintCallable, Category = "Blueprint Upgrade Tools")
+	static void UpgradeOperatorNodes(UBlueprint* Blueprint);
+
+	/**
+	* Compiles the given blueprint. 
+	*
+	* @param Blueprint	Blueprint to compile
+	*/
+	UFUNCTION(BlueprintCallable, Category = "Blueprint Upgrade Tools")
+	static void CompileBlueprint(UBlueprint* Blueprint);
+
+	/**
+	* Adds a function to the given blueprint
+	*
+	* @param Blueprint	The blueprint to remove the function from
+	* @param FuncName	Name of the function to add
+	*
+	* @return UEdGraph*
+	*/
+	UFUNCTION(BlueprintCallable, Category = "Blueprint Upgrade Tools")
+	static UEdGraph* AddFunctionGraph(UBlueprint* Blueprint, const FString& FuncName);
+
+	/** 
+	* Deletes the function of the given name on this blueprint. Does NOT replace function call sites. 
+	*
+	* @param Blueprint		The blueprint to remove the function from
+	* @param FuncName		The name of the function to remove
+	*/
+	UFUNCTION(BlueprintCallable, Category = "Blueprint Upgrade Tools")
+	static void RemoveFunctionGraph(UBlueprint* Blueprint, FName FuncName);
+
+	/**
+	* Remove any nodes in this blueprint that have no connections made to them.
+	*
+	* @param Blueprint		The blueprint to remove the nodes from
+	*/
+	UFUNCTION(BlueprintCallable, Category = "Blueprint Upgrade Tools")
+	static void RemoveUnusedNodes(UBlueprint* Blueprint);
+
+	/** 
+	* Removes the given graph from the blueprint if possible 
+	* 
+	* @param Blueprint	The blueprint the graph will be removed from
+	* @param Graph		The graph to remove
+	*/
+	UFUNCTION(BlueprintCallable, Category = "Blueprint Upgrade Tools")
+	static void RemoveGraph(UBlueprint* Blueprint, UEdGraph* Graph);
+
+	/**
+	* Attempts to rename the given graph with a new name
+	*
+	* @param Graph			The graph to rename
+	* @param NewNameStr		The new name of the graph
+	*/
+	UFUNCTION(BlueprintCallable, Category = "Blueprint Upgrade Tools")
+	static void RenameGraph(UEdGraph* Graph, const FString& NewNameStr);
+
+	/**
+	* Gets the UBlueprint version of the given object if possible.
+	*
+	* @param Object			The object we need to get the UBlueprint from
+	*
+	* @return UBlueprint*	The blueprint type of the given object, nullptr if the object is not a blueprint.
+	*/
+	UFUNCTION(BlueprintCallable, Category = "Blueprint Upgrade Tools")
+	static UBlueprint* GetBlueprintAsset(UObject* Object);
+
+	/**
+	* Attempts to reparent the given blueprint to the new chosen parent class. 
+	*
+	* @param Blueprint			Blueprint that you would like to reparent
+	* @param NewParentClass		The new parent class to use
+	*/
+	UFUNCTION(BlueprintCallable, Category = "Blueprint Upgrade Tools")
+	static void ReparentBlueprint(UBlueprint* Blueprint, UClass* NewParentClass);
+
+	/**
+	* Gathers any unused blueprint variables and populates the given array of FPropertys
+	*
+	* @param Blueprint			The blueprint to check
+	* @param OutProperties		Out array of unused FProperty*'s
+	*
+	* @return					True if variables were checked on this blueprint, false otherwise.
+	*/
+	static bool GatherUnusedVariables(const UBlueprint* Blueprint, TArray<FProperty*>& OutProperties);
+
+	/**
+	* Deletes any unused blueprint created variables the given blueprint.
+	* An Unused variable is any BP variable that is not referenced in any 
+	* blueprint graphs
+	* 
+	* @param Blueprint			Blueprint that you would like to remove variables from
+	*
+	* @return					Number of variables removed
+	*/
+	UFUNCTION(BlueprintCallable, Category = "Blueprint Upgrade Tools")
+	static int32 RemoveUnusedVariables(UBlueprint* Blueprint);
+};
\ No newline at end of file
diff --git a/Engine/Source/Editor/BlueprintEditorLibrary/Public/BlueprintEditorLibraryModule.h b/Engine/Source/Editor/BlueprintEditorLibrary/Public/BlueprintEditorLibraryModule.h
new file mode 100644
--- /dev/null
+++ b/Engine/Source/Editor/BlueprintEditorLibrary/Public/BlueprintEditorLibraryModule.h
@@ -0,0 +1,14 @@
+// Copyright Epic Games, Inc. All Rights Reserved.
+
+#pragma once
+
+#include "Modules/ModuleManager.h"
+
+class FBlueprintEditorLibraryModule : public IModuleInterface
+{
+public:
+
+	/** IModuleInterface implementation */
+	virtual void StartupModule() override;
+	virtual void ShutdownModule() override;
+};
\ No newline at end of file
