remove log spam from blueprint editor library module startup

#rnx

diff --git a/Engine/Source/Editor/BlueprintEditorLibrary/Private/BlueprintEditorLibraryModule.cpp b/Engine/Source/Editor/BlueprintEditorLibrary/Private/BlueprintEditorLibraryModule.cpp
--- a/Engine/Source/Editor/BlueprintEditorLibrary/Private/BlueprintEditorLibraryModule.cpp
+++ b/Engine/Source/Editor/BlueprintEditorLibrary/Private/BlueprintEditorLibraryModule.cpp
@@ -7,14 +7,12 @@
 void FBlueprintEditorLibraryModule::StartupModule()
 {
 	// This code will execute after your module is loaded into memory; the exact timing is specified in the .uplugin file per-module
-	UE_LOG(LogTemp, Warning, TEXT("BP Editor Lib Loaded"));
 }
 
 void FBlueprintEditorLibraryModule::ShutdownModule()
 {
 	// This function may be called during shutdown to clean up your module.  For modules that support dynamic reloading,
 	// we call this function before unloading the module.	
-	UE_LOG(LogTemp, Warning, TEXT("BP Editor Lib Loaded"));
 }
 
 #undef LOCTEXT_NAMESPACE
