Blueprints: Only generate a unique graph name if if fails validation instead of all the time when creating a graph via BP editor library

#rnx

diff --git a/Engine/Source/Editor/BlueprintEditorLibrary/Private/BlueprintEditorLibrary.cpp b/Engine/Source/Editor/BlueprintEditorLibrary/Private/BlueprintEditorLibrary.cpp
--- a/Engine/Source/Editor/BlueprintEditorLibrary/Private/BlueprintEditorLibrary.cpp
+++ b/Engine/Source/Editor/BlueprintEditorLibrary/Private/BlueprintEditorLibrary.cpp
@@ -16,6 +16,7 @@
 #include "EdGraphUtilities.h"
 #include "AnimGraphNode_Base.h"
 #include "Components/TimelineComponent.h"
+#include "Kismet2/Kismet2NameValidators.h"
 
 #define LOCTEXT_NAMESPACE "BlueprintEditorLibrary"
 
@@ -292,10 +293,22 @@ UEdGraph* UBlueprintEditorLibrary::AddFunctionGraph(UBlueprint* Blueprint, const
 		return nullptr;
 	}
 
+	// Validate that the given name is appropriate for a new function graph
+	FName GraphName;
+
+	if (FKismetNameValidator(Blueprint).IsValid(FuncName) == EValidatorResult::Ok)
+	{
+		GraphName = FName(*FuncName);
+	}
+	else
+	{
+		GraphName = FBlueprintEditorUtils::FindUniqueKismetName(Blueprint, FuncName);
+	}
+
 	Blueprint->Modify();
 	UEdGraph* NewGraph = FBlueprintEditorUtils::CreateNewGraph(
 		Blueprint, 
-		FBlueprintEditorUtils::FindUniqueKismetName(Blueprint, FuncName),
+		GraphName,
 		UEdGraph::StaticClass(), 
 		UEdGraphSchema_K2::StaticClass()
 	);
