Update the default value on UBlueprintEditorLibrary::AddFunctionGraph to be "NewFunction". If there is no name given, then default to "NewFunction" when finding a unique name to prevent a hang.

#rnx

diff --git a/Engine/Source/Editor/BlueprintEditorLibrary/Private/BlueprintEditorLibrary.cpp b/Engine/Source/Editor/BlueprintEditorLibrary/Private/BlueprintEditorLibrary.cpp
--- a/Engine/Source/Editor/BlueprintEditorLibrary/Private/BlueprintEditorLibrary.cpp
+++ b/Engine/Source/Editor/BlueprintEditorLibrary/Private/BlueprintEditorLibrary.cpp
@@ -290,6 +290,7 @@ UEdGraph* UBlueprintEditorLibrary::AddFunctionGraph(UBlueprint* Blueprint, const
 {
 	if (!Blueprint)
 	{
+		UE_LOG(LogBlueprintEditorLib, Warning, TEXT("Failed to add function graph, ensure that blueprint is not null!"));
 		return nullptr;
 	}
 
@@ -302,7 +303,8 @@ UEdGraph* UBlueprintEditorLibrary::AddFunctionGraph(UBlueprint* Blueprint, const
 	}
 	else
 	{
-		GraphName = FBlueprintEditorUtils::FindUniqueKismetName(Blueprint, FuncName);
+		static const FString NewFunctionString = TEXT("NewFunction");
+		GraphName = FBlueprintEditorUtils::FindUniqueKismetName(Blueprint, !FuncName.IsEmpty() ? FuncName : NewFunctionString);
 	}
 
 	Blueprint->Modify();
diff --git a/Engine/Source/Editor/BlueprintEditorLibrary/Public/BlueprintEditorLibrary.h b/Engine/Source/Editor/BlueprintEditorLibrary/Public/BlueprintEditorLibrary.h
--- a/Engine/Source/Editor/BlueprintEditorLibrary/Public/BlueprintEditorLibrary.h
+++ b/Engine/Source/Editor/BlueprintEditorLibrary/Public/BlueprintEditorLibrary.h
@@ -86,7 +86,7 @@ public:
 	* @return UEdGraph*
 	*/
 	UFUNCTION(BlueprintCallable, Category = "Blueprint Upgrade Tools")
-	static UEdGraph* AddFunctionGraph(UBlueprint* Blueprint, const FString& FuncName);
+	static UEdGraph* AddFunctionGraph(UBlueprint* Blueprint, const FString& FuncName = FString(TEXT("NewFunction")));
 
 	/** 
 	* Deletes the function of the given name on this blueprint. Does NOT replace function call sites. 
