Remove unused function "ReplaceK2Nodes"

#rnx

diff --git a/Engine/Source/Editor/BlueprintEditorLibrary/Private/BlueprintEditorLibrary.cpp b/Engine/Source/Editor/BlueprintEditorLibrary/Private/BlueprintEditorLibrary.cpp
--- a/Engine/Source/Editor/BlueprintEditorLibrary/Private/BlueprintEditorLibrary.cpp
+++ b/Engine/Source/Editor/BlueprintEditorLibrary/Private/BlueprintEditorLibrary.cpp
@@ -110,42 +110,6 @@ void UBlueprintEditorLibrary::ReplaceVariableReferences(UBlueprint* Blueprint, c
 	FBlueprintEditorUtils::RenameVariableReferences(Blueprint, Blueprint->GeneratedClass, OldVarName, NewVarName);
 }
 
-void UBlueprintEditorLibrary::ReplaceK2Nodes(UBlueprint* Blueprint, TSubclassOf<UEdGraphNode> OldNodeType, TSubclassOf<UEdGraphNode> NewNodeType)
-{
-	if (!Blueprint)
-	{
-		return;
-	}
-
-	TArray<UEdGraph*> AllGraphs;
-	Blueprint->GetAllGraphs(AllGraphs);
-
-	for (UEdGraph* Graph : AllGraphs)
-	{
-		Graph->Modify();
-
-		// Iterate backwards through the nodes because we will be removing some
-		for (int32 i = Graph->Nodes.Num() - 1; i >= 0; i--)
-		{
-			UEdGraphNode* OriginalNode = Graph->Nodes[i];
-			if (OriginalNode && OriginalNode->IsA(OldNodeType))
-			{
-				// Spawn a new node of the new class
-				UEdGraphNode* NewNode = NewObject<UEdGraphNode>(Graph, NewNodeType);
-				NewNode->CreateNewGuid();
-				NewNode->PostPlacedNewNode();
-				NewNode->AllocateDefaultPins();
-				NewNode->SetFlags(RF_Transactional);
-
-				NewNode->NodePosX = OriginalNode->NodePosX;
-				NewNode->NodePosY = OriginalNode->NodePosY;
-
-				InternalBlueprintEditorLibrary::ReplaceOldNodeWithNew(OriginalNode, NewNode);
-			}
-		}
-	}
-}
-
 UEdGraph* UBlueprintEditorLibrary::FindEventGraph(UBlueprint* Blueprint)
 {
 	return Blueprint ? FBlueprintEditorUtils::FindEventGraph(Blueprint) : nullptr;
diff --git a/Engine/Source/Editor/BlueprintEditorLibrary/Public/BlueprintEditorLibrary.h b/Engine/Source/Editor/BlueprintEditorLibrary/Public/BlueprintEditorLibrary.h
--- a/Engine/Source/Editor/BlueprintEditorLibrary/Public/BlueprintEditorLibrary.h
+++ b/Engine/Source/Editor/BlueprintEditorLibrary/Public/BlueprintEditorLibrary.h
@@ -27,16 +27,6 @@ public:
 	UFUNCTION(BlueprintCallable, Category = "Blueprint Upgrade Tools")
 	static void ReplaceVariableReferences(UBlueprint* Blueprint, const FName OldVarName, const FName NewVarName);
 
-	/**
-	* Replace any existing k2node class with the new one, attempting to remap pins as we go
-	*
-	* @param Blueprint		Blueprint to replace node types on
-	* @param OldNodeType	Old node type that will be removed
-	* @param NewNodeType	New node type that will be put in the old node's place
-	*/
-	UFUNCTION(BlueprintCallable, Category = "Blueprint Upgrade Tools")
-	static void ReplaceK2Nodes(UBlueprint* Blueprint, TSubclassOf<UEdGraphNode> OldNodeType, TSubclassOf<UEdGraphNode> NewNodeType);
-
 	/**
 	* Finds the event graph of the given blueprint. Null if it doesn't have one. This will only return
 	* the primary event graph of the blueprint (the graph named "EventGraph").
