Close any open asset editors when maknig modifications to a blueprint asset via the BP Editor Library

#rnx

diff --git a/Engine/Source/Editor/BlueprintEditorLibrary/Private/BlueprintEditorLibrary.cpp b/Engine/Source/Editor/BlueprintEditorLibrary/Private/BlueprintEditorLibrary.cpp
--- a/Engine/Source/Editor/BlueprintEditorLibrary/Private/BlueprintEditorLibrary.cpp
+++ b/Engine/Source/Editor/BlueprintEditorLibrary/Private/BlueprintEditorLibrary.cpp
@@ -17,6 +17,7 @@
 #include "AnimGraphNode_Base.h"
 #include "Components/TimelineComponent.h"
 #include "Kismet2/Kismet2NameValidators.h"
+#include "Subsystems/AssetEditorSubsystem.h"
 
 #define LOCTEXT_NAMESPACE "BlueprintEditorLibrary"
 
@@ -90,6 +91,21 @@ namespace InternalBlueprintEditorLibrary
 
 		return false;
 	}
+
+	/**
+	* Attempt to close any open editors that may be relevant to this blueprint. This will prevent any 
+	* problems where the user could see a previously deleted node/graph.
+	*
+	* @param Blueprint		The blueprint that is being edited
+	*/
+	static void CloseOpenEditors(UBlueprint* Blueprint)
+	{
+		UAssetEditorSubsystem* AssetSubsystem = GEditor ? GEditor->GetEditorSubsystem<UAssetEditorSubsystem>() : nullptr;
+		if (AssetSubsystem && Blueprint)
+		{
+			AssetSubsystem->CloseAllEditorsForAsset(Blueprint);
+		}
+	}
 };
 
 ///////////////////////////////////////////////////////////
@@ -306,6 +322,7 @@ void UBlueprintEditorLibrary::RemoveFunctionGraph(UBlueprint* Blueprint, FName F
 	if (FunctionGraph && FunctionGraph->bAllowDeletion)
 	{
 		Blueprint->Modify();
+		InternalBlueprintEditorLibrary::CloseOpenEditors(Blueprint);
 		FBlueprintEditorUtils::RemoveGraph(Blueprint, FunctionGraph, EGraphRemoveFlags::MarkTransient);
 	}
 	else
@@ -370,6 +387,7 @@ void UBlueprintEditorLibrary::RemoveGraph(UBlueprint* Blueprint, UEdGraph* Graph
 		return;
 	}
 
+	InternalBlueprintEditorLibrary::CloseOpenEditors(Blueprint);
 	FBlueprintEditorUtils::RemoveGraph(Blueprint, Graph, EGraphRemoveFlags::MarkTransient);
 }
 
