When renaming a graph via the BP editor library, correctly generate a unique kismet name if necessary. Update the default value to be more user friendly. If an empty name is given then default to "NewGraph"

#rnx

diff --git a/Engine/Source/Editor/BlueprintEditorLibrary/Private/BlueprintEditorLibrary.cpp b/Engine/Source/Editor/BlueprintEditorLibrary/Private/BlueprintEditorLibrary.cpp
--- a/Engine/Source/Editor/BlueprintEditorLibrary/Private/BlueprintEditorLibrary.cpp
+++ b/Engine/Source/Editor/BlueprintEditorLibrary/Private/BlueprintEditorLibrary.cpp
@@ -395,10 +395,31 @@ void UBlueprintEditorLibrary::RenameGraph(UEdGraph* Graph, const FString& NewNam
 {
 	if (!Graph)
 	{
+		UE_LOG(LogBlueprintEditorLib, Warning, TEXT("Invalid graph given, failed to rename!"));
 		return;
 	}
-		
-	FBlueprintEditorUtils::RenameGraph(Graph, NewNameStr);
+	
+	// Validate that the given name is appropriate for a new function graph
+	UBlueprint* BP = FBlueprintEditorUtils::FindBlueprintForGraph(Graph);
+	if (!BP)
+	{
+		UE_LOG(LogBlueprintEditorLib, Warning, TEXT("Failed to find blueprint for graph!"));
+		return;
+	}
+
+	FString ValidatedNewName;
+
+	if (FKismetNameValidator(BP).IsValid(NewNameStr) == EValidatorResult::Ok)
+	{
+		ValidatedNewName = NewNameStr;
+	}
+	else
+	{
+		static const FString RenamedGraphString = TEXT("NewGraph");
+		ValidatedNewName = FBlueprintEditorUtils::FindUniqueKismetName(BP, !NewNameStr.IsEmpty() ? NewNameStr : RenamedGraphString).ToString();
+	}
+
+	FBlueprintEditorUtils::RenameGraph(Graph, ValidatedNewName);
 }
 
 UBlueprint* UBlueprintEditorLibrary::GetBlueprintAsset(UObject* Object)
diff --git a/Engine/Source/Editor/BlueprintEditorLibrary/Public/BlueprintEditorLibrary.h b/Engine/Source/Editor/BlueprintEditorLibrary/Public/BlueprintEditorLibrary.h
--- a/Engine/Source/Editor/BlueprintEditorLibrary/Public/BlueprintEditorLibrary.h
+++ b/Engine/Source/Editor/BlueprintEditorLibrary/Public/BlueprintEditorLibrary.h
@@ -111,7 +111,7 @@ public:
 	* @param NewNameStr		The new name of the graph
 	*/
 	UFUNCTION(BlueprintCallable, Category = "Blueprint Upgrade Tools")
-	static void RenameGraph(UEdGraph* Graph, const FString& NewNameStr);
+	static void RenameGraph(UEdGraph* Graph, const FString& NewNameStr = FString(TEXT("NewGraph")));
 
 	/**
 	* Gets the UBlueprint version of the given object if possible.
