Do not remove comment nodes when removing unused nodes

#rnx

diff --git a/Engine/Source/Editor/BlueprintEditorLibrary/Private/BlueprintEditorLibrary.cpp b/Engine/Source/Editor/BlueprintEditorLibrary/Private/BlueprintEditorLibrary.cpp
--- a/Engine/Source/Editor/BlueprintEditorLibrary/Private/BlueprintEditorLibrary.cpp
+++ b/Engine/Source/Editor/BlueprintEditorLibrary/Private/BlueprintEditorLibrary.cpp
@@ -362,6 +362,7 @@ void UBlueprintEditorLibrary::RemoveUnusedNodes(UBlueprint* Blueprint)
 			// be valid with no pin connections made to them
 			if (Node->CanUserDeleteNode() && 
 				!Node->IsA<UAnimGraphNode_Base>() && 
+				!Node->IsA<UEdGraphNode_Comment>() &&
 				!InternalBlueprintEditorLibrary::NodeHasAnyConnections(Node))
 			{
 				Node->BreakAllNodeLinks();
