Add a return value to have better messaging when replacing nodes with the BP editor library. Change a debug if statement that would refresh all actions upon upgrading, which is no longer needed. Reset promotable operators with no connections to wildcard state when upgrading for a more cohesive upgrade.

diff --git a/Engine/Source/Editor/BlueprintEditorLibrary/Private/BlueprintEditorLibrary.cpp b/Engine/Source/Editor/BlueprintEditorLibrary/Private/BlueprintEditorLibrary.cpp
--- a/Engine/Source/Editor/BlueprintEditorLibrary/Private/BlueprintEditorLibrary.cpp
+++ b/Engine/Source/Editor/BlueprintEditorLibrary/Private/BlueprintEditorLibrary.cpp
@@ -35,9 +35,11 @@ namespace InternalBlueprintEditorLibrary
 	* @param OldNode		The old node to replace
 	* @param NewNode		The new node to put in the old node's place
 	*/
-	static void ReplaceOldNodeWithNew(UEdGraphNode* OldNode, UEdGraphNode* NewNode)
+	static bool ReplaceOldNodeWithNew(UEdGraphNode* OldNode, UEdGraphNode* NewNode)
 	{
 		const UEdGraphSchema_K2* Schema = GetDefault<UEdGraphSchema_K2>();
+		
+		bool bSuccess = false;
 
 		if (Schema && OldNode && NewNode)
 		{
@@ -61,11 +63,13 @@ namespace InternalBlueprintEditorLibrary
 				}
 			}
 			
-			Schema->ReplaceOldNodeWithNew(OldNode, NewNode, OldToNewPinMap);
+			bSuccess = Schema->ReplaceOldNodeWithNew(OldNode, NewNode, OldToNewPinMap);
 			// reconstructing the node will clean up any
 			// incorrect default values that may have been copied over
-			NewNode->ReconstructNode();
+			NewNode->ReconstructNode();			
 		}
+
+		return bSuccess;
 	}
 
 	/**
@@ -185,16 +189,7 @@ void UBlueprintEditorLibrary::UpgradeOperatorNodes(UBlueprint* Blueprint)
 		return;
 	}
 
-	if (TypePromoDebug::IsTypePromoEnabled())
-	{
-		// Ensure that we have promotable operator node spanners available to us. 
-		// They will be empty if the editor hasn't been opened or 
-		if (FBlueprintActionDatabase* Actions = FBlueprintActionDatabase::TryGet())
-		{
-			Actions->RefreshAll();
-		}
-	}
-	else
+	if (!TypePromoDebug::IsTypePromoEnabled())
 	{
 		UE_LOG(LogBlueprintEditorLib, Warning, TEXT("Type Promotion is not enabled! Cannot upgrade operator nodes. Set 'BP.TypePromo.IsEnabled' to true and try again."));
 		return;
@@ -218,6 +213,7 @@ void UBlueprintEditorLibrary::UpgradeOperatorNodes(UBlueprint* Blueprint)
 			{
 				UFunction* Func = OldOpNode->GetTargetFunction();
 				UEdGraph* OwningGraph = OldOpNode->GetGraph();
+				const bool bHadAnyConnections = InternalBlueprintEditorLibrary::NodeHasAnyConnections(OldOpNode);
 
 				// We should only be modifying nodes within the graph that we want
 				ensure(OwningGraph == Graph);
@@ -251,6 +247,15 @@ void UBlueprintEditorLibrary::UpgradeOperatorNodes(UBlueprint* Blueprint)
 				NewOpNode->NodePosY = OldOpNode->NodePosY;
 
 				InternalBlueprintEditorLibrary::ReplaceOldNodeWithNew(OldOpNode, NewOpNode);
+
+				// Reset the new node to be wild card if there were no connections to the original node.
+				// This is necessary because replacing the old node will attempt to reconcile any 
+				// default values on the node, which can result in incorrect pin types and a default
+				// value that doesn't match. 
+				if(!bHadAnyConnections)
+				{
+					NewOpNode->ResetNodeToWildcard();
+				}
 			}
 		}
 	}
diff --git a/Engine/Source/Editor/BlueprintGraph/Classes/K2Node_PromotableOperator.h b/Engine/Source/Editor/BlueprintGraph/Classes/K2Node_PromotableOperator.h
--- a/Engine/Source/Editor/BlueprintGraph/Classes/K2Node_PromotableOperator.h
+++ b/Engine/Source/Editor/BlueprintGraph/Classes/K2Node_PromotableOperator.h
@@ -52,6 +52,9 @@ public:
 	// UK2Node_CallFunction interface
 	virtual void SetFromFunction(const UFunction* Function);
 	// End of UK2Node_CallFunction interface
+
+	/** Recombines all split pins and sets the node to have default values (all wildcard pins) */
+	BLUEPRINTGRAPH_API void ResetNodeToWildcard();
 	
 private:
 
@@ -87,9 +90,6 @@ private:
 	*/
 	bool HasDeterminingComparisonTypes() const;
 
-	/** Recombines all split pins and sets the node to have default values (all wildcard pins) */
-	void ResetNodeToWildcard();
-
 	/**
 	* Attempts to create a cast node and connect it to a call function node
 	* 
