Correctly restore the default values of math nodes when running them through the BP editor library for an upgrade.

#rnx

diff --git a/Engine/Source/Editor/BlueprintEditorLibrary/Private/BlueprintEditorLibrary.cpp b/Engine/Source/Editor/BlueprintEditorLibrary/Private/BlueprintEditorLibrary.cpp
--- a/Engine/Source/Editor/BlueprintEditorLibrary/Private/BlueprintEditorLibrary.cpp
+++ b/Engine/Source/Editor/BlueprintEditorLibrary/Private/BlueprintEditorLibrary.cpp
@@ -199,6 +199,23 @@ void UBlueprintEditorLibrary::UpgradeOperatorNodes(UBlueprint* Blueprint)
 	Blueprint->GetAllGraphs(AllGraphs);
 	Blueprint->Modify();
 
+	/**
+	* Used to help us restore the default values of any pins that may have changed their types
+	* during replacement. 
+	*/
+	struct FRestoreDefaultsHelper
+	{
+		FEdGraphPinType PinType {};
+
+		FString DefaultValue = TEXT("");
+
+		TObjectPtr<UObject> DefaultObject = nullptr;
+
+		FText DefaultTextValue = FText::GetEmpty();
+	};
+
+	TMap<FName, FRestoreDefaultsHelper> PinTypeMap;
+
 	for (UEdGraph* Graph : AllGraphs)
 	{
 		check(Graph);
@@ -207,6 +224,8 @@ void UBlueprintEditorLibrary::UpgradeOperatorNodes(UBlueprint* Blueprint)
 
 		for (int32 i = Graph->Nodes.Num() - 1; i >= 0; --i)
 		{
+			PinTypeMap.Reset();
+			
 			// Not every function that we want to upgrade is a CommunicativeBinaryOpNode
 			// Some are just regular CallFunction nodes; Vector + Float is an example of this
 			if (UK2Node_CallFunction* OldOpNode = Cast<UK2Node_CallFunction>(Graph->Nodes[i]))
@@ -224,6 +243,21 @@ void UBlueprintEditorLibrary::UpgradeOperatorNodes(UBlueprint* Blueprint)
 					continue;
 				}
 
+				// Keep track of the types of anything with a default value so they can be restored
+				for (UEdGraphPin* Pin : OldOpNode->Pins)
+				{
+					if (Pin->Direction == EGPD_Input && Pin->LinkedTo.IsEmpty())
+					{
+						FRestoreDefaultsHelper RestoreData;
+						RestoreData.PinType = Pin->PinType;
+						RestoreData.DefaultValue = Pin->DefaultValue;
+						RestoreData.DefaultObject = Pin->DefaultObject;
+						RestoreData.DefaultTextValue = Pin->DefaultTextValue;
+
+						PinTypeMap.Add(Pin->GetFName(), RestoreData);
+					}
+				}
+
 				FName OpName = FTypePromotion::GetOpNameFromFunction(Func);
 
 				UK2Node_CommutativeAssociativeBinaryOperator* BinaryOpNode = Cast<UK2Node_CommutativeAssociativeBinaryOperator>(OldOpNode);
@@ -248,6 +282,22 @@ void UBlueprintEditorLibrary::UpgradeOperatorNodes(UBlueprint* Blueprint)
 
 				InternalBlueprintEditorLibrary::ReplaceOldNodeWithNew(OldOpNode, NewOpNode);
 
+				for (const TPair<FName, FRestoreDefaultsHelper>& Pair : PinTypeMap)
+				{
+					const FRestoreDefaultsHelper& OldPinData = Pair.Value;
+
+					if (UEdGraphPin* Pin = NewOpNode->FindPin(Pair.Key))
+					{
+						if (NewOpNode->CanConvertPinType(Pin))
+						{
+							NewOpNode->ConvertPinType(Pin, OldPinData.PinType);
+							Pin->DefaultValue = OldPinData.DefaultValue;
+							Pin->DefaultObject = OldPinData.DefaultObject;
+							Pin->DefaultTextValue = OldPinData.DefaultTextValue;
+						}
+					}
+				}
+
 				// Reset the new node to be wild card if there were no connections to the original node.
 				// This is necessary because replacing the old node will attempt to reconcile any 
 				// default values on the node, which can result in incorrect pin types and a default
diff --git a/Engine/Source/Editor/BlueprintGraph/Classes/K2Node_PromotableOperator.h b/Engine/Source/Editor/BlueprintGraph/Classes/K2Node_PromotableOperator.h
--- a/Engine/Source/Editor/BlueprintGraph/Classes/K2Node_PromotableOperator.h
+++ b/Engine/Source/Editor/BlueprintGraph/Classes/K2Node_PromotableOperator.h
@@ -137,16 +137,6 @@ private:
 	* @param OutArray	Array to populate with the considered pins
 	*/
 	void GetPinsToConsider(TArray<UEdGraphPin*>& OutArray) const;
-
-	/**
-	* Convert the given pin to the new pin type. If the new type is wildcard then break 
-	* all connections and reset the node. Invalidates tooltips upon completion. 
-	* This function is to be used by the context menu and includes a transaction. 
-	* 
-	* @param PinToChange		The Pin to convert
-	* @param NewPinType			The pin's new type
-	*/
-	void ConvertPinType(UEdGraphPin* PinToChange, const FEdGraphPinType NewPinType);
 	
 	/**
 	* Build the context menu to convert the given pin to a different type. 
@@ -158,15 +148,6 @@ private:
 	*/
 	void CreateConversionSubMenu(UToolMenu* Menu, UEdGraphPin* ContextPin) const;
 
-	/**
-	* Returns true if the given pin can be converted via the context menu to another type
-	*
-	* @param Pin		The pin to consider
-	*
-	* @return bool		True if a conversion is possible
-	*/
-	bool CanConvertPinType(const UEdGraphPin* Pin) const;
-
 	/** The name that this operation uses ("Add", "Multiply", etc) */
 	FName OperationName;
 
@@ -183,4 +164,23 @@ public:
 	BLUEPRINTGRAPH_API UEdGraphPin* GetOutputPin() const;
 
 	const FName GetOperationName() const { return OperationName; }
+
+	/**
+	* Convert the given pin to the new pin type. If the new type is wildcard then break 
+	* all connections and reset the node. Invalidates tooltips upon completion. 
+	* This function is to be used by the context menu and includes a transaction. 
+	* 
+	* @param PinToChange		The Pin to convert
+	* @param NewPinType			The pin's new type
+	*/
+	BLUEPRINTGRAPH_API void ConvertPinType(UEdGraphPin* PinToChange, const FEdGraphPinType NewPinType);
+
+	/**
+	* Returns true if the given pin can be converted via the context menu to another type
+	*
+	* @param Pin		The pin to consider
+	*
+	* @return bool		True if a conversion is possible
+	*/
+	BLUEPRINTGRAPH_API bool CanConvertPinType(const UEdGraphPin* Pin) const;
 };
