PR #8175: Fixed incorrect doc-comment for AddFunctionGraph function. (Contributed by [REDACTED_PERSON])

diff --git a/Engine/Source/Editor/BlueprintEditorLibrary/Public/BlueprintEditorLibrary.h b/Engine/Source/Editor/BlueprintEditorLibrary/Public/BlueprintEditorLibrary.h
--- a/Engine/Source/Editor/BlueprintEditorLibrary/Public/BlueprintEditorLibrary.h
+++ b/Engine/Source/Editor/BlueprintEditorLibrary/Public/BlueprintEditorLibrary.h
@@ -70,7 +70,7 @@ public:
 	/**
 	* Adds a function to the given blueprint
 	*
-	* @param Blueprint	The blueprint to remove the function from
+	* @param Blueprint	The blueprint to add the function to
 	* @param FuncName	Name of the function to add
 	*
 	* @return UEdGraph*
