BlueprintEditorLibrary: Added additional editor scripting APIs to modify variable flags/metadata

diff --git a/Engine/Source/Editor/BlueprintEditorLibrary/Private/BlueprintEditorLibrary.cpp b/Engine/Source/Editor/BlueprintEditorLibrary/Private/BlueprintEditorLibrary.cpp
--- a/Engine/Source/Editor/BlueprintEditorLibrary/Private/BlueprintEditorLibrary.cpp
+++ b/Engine/Source/Editor/BlueprintEditorLibrary/Private/BlueprintEditorLibrary.cpp
@@ -601,4 +601,61 @@ UClass* UBlueprintEditorLibrary::GeneratedClass(UBlueprint* BlueprintObj)
 	return nullptr;
 }
 
-#undef LOCTEXT_NAMESPACE	// "BlueprintEditorLibrary"
\ No newline at end of file
+void UBlueprintEditorLibrary::SetBlueprintVariableExposeOnSpawn(UBlueprint* Blueprint, const FName& VariableName, bool bExposeOnSpawn)
+{
+	if (!Blueprint)
+	{
+		UE_LOG(LogBlueprintEditorLib, Warning, TEXT("Invalid Blueprint!"));
+		return;
+	}
+	
+	if (VariableName == NAME_None)
+	{
+		UE_LOG(LogBlueprintEditorLib, Warning, TEXT("Invalid variable name!"));
+		return;
+	}
+	
+	if(bExposeOnSpawn)
+	{
+		FBlueprintEditorUtils::SetBlueprintVariableMetaData(Blueprint, VariableName, NULL, FBlueprintMetadata::MD_ExposeOnSpawn, TEXT("true"));
+	}
+	else
+	{
+		FBlueprintEditorUtils::RemoveBlueprintVariableMetaData(Blueprint, VariableName, NULL, FBlueprintMetadata::MD_ExposeOnSpawn);
+	} 
+}
+
+void UBlueprintEditorLibrary::SetBlueprintVariableExposeToCinematics(UBlueprint* Blueprint, const FName& VariableName, bool bExposeToCinematics)
+{
+	if (!Blueprint)
+	{
+		UE_LOG(LogBlueprintEditorLib, Warning, TEXT("Invalid Blueprint!"));
+		return;
+	}
+	
+	if (VariableName == NAME_None)
+	{
+		UE_LOG(LogBlueprintEditorLib, Warning, TEXT("Invalid variable name!"));
+		return;
+	}
+
+	FBlueprintEditorUtils::SetInterpFlag(Blueprint, VariableName, bExposeToCinematics);
+}
+
+void UBlueprintEditorLibrary::SetBlueprintVariableInstanceEditable(UBlueprint* Blueprint, const FName& VariableName, bool bInstanceEditable)
+{
+	if (!Blueprint)
+	{
+		UE_LOG(LogBlueprintEditorLib, Warning, TEXT("Invalid Blueprint!"));
+		return;
+	}
+	
+	if (VariableName == NAME_None)
+	{
+		UE_LOG(LogBlueprintEditorLib, Warning, TEXT("Invalid variable name!"));
+		return;
+	}
+	
+	FBlueprintEditorUtils::SetBlueprintOnlyEditableFlag(Blueprint, VariableName, !bInstanceEditable);
+}
+#undef LOCTEXT_NAMESPACE	// "BlueprintEditorLibrary"
diff --git a/Engine/Source/Editor/BlueprintEditorLibrary/Public/BlueprintEditorLibrary.h b/Engine/Source/Editor/BlueprintEditorLibrary/Public/BlueprintEditorLibrary.h
--- a/Engine/Source/Editor/BlueprintEditorLibrary/Public/BlueprintEditorLibrary.h
+++ b/Engine/Source/Editor/BlueprintEditorLibrary/Public/BlueprintEditorLibrary.h
@@ -162,4 +162,34 @@ public:
 	 */
 	UFUNCTION(BlueprintCallable, Category = "Editor Scripting", meta = (ScriptMethod))
 	static UClass* GeneratedClass(UBlueprint* BlueprintObj);
+
+	/**
+	 * Sets "Expose On Spawn" to true/false on a Blueprint variable
+	 *
+	 * @param Blueprint			The blueprint object
+	 * @param VariableName		The variable name
+	 * @param bExposeOnSpawn	Set to true to expose on spawn
+	 */
+	UFUNCTION(BlueprintCallable, Category = "Editor Scripting", meta = (ScriptMethod))
+	static void SetBlueprintVariableExposeOnSpawn(UBlueprint* Blueprint, const FName& VariableName, bool bExposeOnSpawn);
+
+	/**
+	 * Sets "Expose To Cinematics" to true/false on a Blueprint variable
+	 *
+	 * @param Blueprint				The blueprint object
+	 * @param VariableName			The variable name
+	 * @param bExposeToCinematics	Set to true to expose to cinematics
+	 */
+	UFUNCTION(BlueprintCallable, Category = "Editor Scripting", meta = (ScriptMethod))
+	static void SetBlueprintVariableExposeToCinematics(UBlueprint* Blueprint, const FName& VariableName, bool bExposeToCinematics);
+
+	/**
+	 * Sets "Instance Editable" to true/false on a Blueprint variable
+	 *
+	 * @param Blueprint				The blueprint object
+	 * @param VariableName			The variable name
+	 * @param bInstanceEditable		Toggle InstanceEditable
+	 */
+	UFUNCTION(BlueprintCallable, Category = "Editor Scripting", meta = (ScriptMethod))
+	static void SetBlueprintVariableInstanceEditable(UBlueprint* Blueprint, const FName& VariableName, bool bInstanceEditable);
 };
\ No newline at end of file
