Add Refresh blueprint editor functions to the BlueprintEditorLibrary. This will allow you to refresh any open blueprint editors oronly ones for a specific blueprint asset when making changes to their properties from editor utiltity widgets

diff --git a/Engine/Source/Editor/BlueprintEditorLibrary/Private/BlueprintEditorLibrary.cpp b/Engine/Source/Editor/BlueprintEditorLibrary/Private/BlueprintEditorLibrary.cpp
--- a/Engine/Source/Editor/BlueprintEditorLibrary/Private/BlueprintEditorLibrary.cpp
+++ b/Engine/Source/Editor/BlueprintEditorLibrary/Private/BlueprintEditorLibrary.cpp
@@ -2,6 +2,8 @@
 
 #include "BlueprintEditorLibrary.h"
 
+#include "BlueprintEditorModule.h"
+#include "BlueprintEditor.h"
 #include "AnimGraphNode_Base.h"
 #include "BlueprintFunctionNodeSpawner.h"
 #include "BlueprintNodeBinder.h"
@@ -505,6 +507,31 @@ UBlueprint* UBlueprintEditorLibrary::GetBlueprintAsset(UObject* Object)
 	return Cast<UBlueprint>(Object);
 }
 
+void UBlueprintEditorLibrary::RefreshOpenEditorsForBlueprint(const UBlueprint* BP)
+{
+	// Get any open blueprint editors for this asset and refresh them if they match the given blueprint
+	FBlueprintEditorModule& BlueprintEditorModule = FModuleManager::LoadModuleChecked<FBlueprintEditorModule>("Kismet");
+	for (TSharedRef<IBlueprintEditor>& Editor : BlueprintEditorModule.GetBlueprintEditors())
+	{
+		if (TSharedPtr<FBlueprintEditor> BPEditor = StaticCastSharedPtr<FBlueprintEditor>(Editor.ToSharedPtr()))
+		{
+			if (BPEditor->GetBlueprintObj() == BP)
+			{
+				BPEditor->RefreshEditors();
+			}
+		}
+	}
+}
+
+void UBlueprintEditorLibrary::RefreshAllOpenBlueprintEditors()
+{
+	FBlueprintEditorModule& BlueprintEditorModule = FModuleManager::LoadModuleChecked<FBlueprintEditorModule>("Kismet");
+	for (TSharedRef<IBlueprintEditor>& Editor : BlueprintEditorModule.GetBlueprintEditors())
+	{
+		Editor->RefreshEditors();
+	}
+}
+
 void UBlueprintEditorLibrary::ReparentBlueprint(UBlueprint* Blueprint, UClass* NewParentClass)
 {
 	if (!Blueprint || !NewParentClass)
diff --git a/Engine/Source/Editor/BlueprintEditorLibrary/Public/BlueprintEditorLibrary.h b/Engine/Source/Editor/BlueprintEditorLibrary/Public/BlueprintEditorLibrary.h
--- a/Engine/Source/Editor/BlueprintEditorLibrary/Public/BlueprintEditorLibrary.h
+++ b/Engine/Source/Editor/BlueprintEditorLibrary/Public/BlueprintEditorLibrary.h
@@ -137,6 +137,14 @@ public:
 	UFUNCTION(BlueprintCallable, Category = "Blueprint Upgrade Tools")
 	static UBlueprint* GetBlueprintAsset(UObject* Object);
 
+	/** Attempt to refresh any open blueprint editors for the given asset */
+	UFUNCTION(BlueprintCallable, Category = "Blueprint Editor")
+	static void RefreshOpenEditorsForBlueprint(const UBlueprint* BP);
+
+	/** Refresh any open blueprint editors */
+	UFUNCTION(BlueprintCallable, Category = "Blueprint Editor")
+	static void RefreshAllOpenBlueprintEditors();
+
 	/**
 	* Attempts to reparent the given blueprint to the new chosen parent class. 
 	*
