Fixed ReparentBlueprint to properly set blueprint compiler options

#UE-185007

diff --git a/Engine/Source/Editor/BlueprintEditorLibrary/Private/BlueprintEditorLibrary.cpp b/Engine/Source/Editor/BlueprintEditorLibrary/Private/BlueprintEditorLibrary.cpp
--- a/Engine/Source/Editor/BlueprintEditorLibrary/Private/BlueprintEditorLibrary.cpp
+++ b/Engine/Source/Editor/BlueprintEditorLibrary/Private/BlueprintEditorLibrary.cpp
@@ -552,13 +552,26 @@ void UBlueprintEditorLibrary::ReparentBlueprint(UBlueprint* Blueprint, UClass* N
 		UE_LOG(LogBlueprintEditorLib, Warning, TEXT("'%s' class heirarcy is changing, there could be possible data loss!"), *Blueprint->GetFriendlyName());
 	}
 
-	UClass* OriginalParentClass = Blueprint->ParentClass;
 	Blueprint->ParentClass = NewParentClass;
 
 	FBlueprintEditorUtils::RefreshAllNodes(Blueprint);
 	FBlueprintEditorUtils::MarkBlueprintAsModified(Blueprint);
 
-	CompileBlueprint(Blueprint);
+	EBlueprintCompileOptions CompileOptions
+	{
+			EBlueprintCompileOptions::SkipSave
+		|	EBlueprintCompileOptions::UseDeltaSerializationDuringReinstancing
+		|	EBlueprintCompileOptions::SkipNewVariableDefaultsDetection
+	};
+
+	// If compilation is enabled during PIE/simulation, references to the CDO might be held by a script variable.
+	// Thus, we set the flag to direct the compiler to allow those references to be replaced during reinstancing.
+	if (GEditor && GEditor->PlayWorld != nullptr)
+	{
+		CompileOptions |= EBlueprintCompileOptions::IncludeCDOInReferenceReplacement;
+	}
+
+	FKismetEditorUtilities::CompileBlueprint(Blueprint, CompileOptions);
 }
 
 bool UBlueprintEditorLibrary::GatherUnusedVariables(const UBlueprint* Blueprint, TArray<FProperty*>& OutProperties)
