Improve discoverability of GetBlueprintAsset
#rnx

diff --git a/Engine/Source/Editor/BlueprintEditorLibrary/Public/BlueprintEditorLibrary.h b/Engine/Source/Editor/BlueprintEditorLibrary/Public/BlueprintEditorLibrary.h
--- a/Engine/Source/Editor/BlueprintEditorLibrary/Public/BlueprintEditorLibrary.h
+++ b/Engine/Source/Editor/BlueprintEditorLibrary/Public/BlueprintEditorLibrary.h
@@ -128,13 +128,14 @@ public:
 	static void RenameGraph(UEdGraph* Graph, const FString& NewNameStr = FString(TEXT("NewGraph")));
 
 	/**
-	* Gets the UBlueprint version of the given object if possible.
+	* Casts the provided Object to a Blueprint - the root asset type of a blueprint asset. Note
+	* that the blueprint asset itself is editor only and not present in cooked assets.
 	*
 	* @param Object			The object we need to get the UBlueprint from
 	*
 	* @return UBlueprint*	The blueprint type of the given object, nullptr if the object is not a blueprint.
 	*/
-	UFUNCTION(BlueprintCallable, Category = "Blueprint Upgrade Tools")
+	UFUNCTION(BlueprintCallable, Category = "Blueprint Upgrade Tools", meta = (Keywords = "cast"))
 	static UBlueprint* GetBlueprintAsset(UObject* Object);
 
 	/** Attempt to refresh any open blueprint editors for the given asset */
