Typo fix, hierarcy->hierarchy
#rnx

diff --git a/Engine/Source/Editor/BlueprintEditorLibrary/Private/BlueprintEditorLibrary.cpp b/Engine/Source/Editor/BlueprintEditorLibrary/Private/BlueprintEditorLibrary.cpp
--- a/Engine/Source/Editor/BlueprintEditorLibrary/Private/BlueprintEditorLibrary.cpp
+++ b/Engine/Source/Editor/BlueprintEditorLibrary/Private/BlueprintEditorLibrary.cpp
@@ -549,7 +549,7 @@ void UBlueprintEditorLibrary::ReparentBlueprint(UBlueprint* Blueprint, UClass* N
 	// There could be possible data loss if reparenting outside the current class hierarchy
 	if (!Blueprint->ParentClass || !NewParentClass->GetDefaultObject()->IsA(Blueprint->ParentClass))
 	{
-		UE_LOG(LogBlueprintEditorLib, Warning, TEXT("'%s' class heirarcy is changing, there could be possible data loss!"), *Blueprint->GetFriendlyName());
+		UE_LOG(LogBlueprintEditorLib, Warning, TEXT("'%s' class hierarchy is changing, there could be possible data loss!"), *Blueprint->GetFriendlyName());
 	}
 
 	Blueprint->ParentClass = NewParentClass;
