ValidateSceneRootNodes after updating the parent class using ReparentBlueprint

diff --git a/Engine/Source/Editor/BlueprintEditorLibrary/Private/BlueprintEditorLibrary.cpp b/Engine/Source/Editor/BlueprintEditorLibrary/Private/BlueprintEditorLibrary.cpp
--- a/Engine/Source/Editor/BlueprintEditorLibrary/Private/BlueprintEditorLibrary.cpp
+++ b/Engine/Source/Editor/BlueprintEditorLibrary/Private/BlueprintEditorLibrary.cpp
@@ -19,6 +19,7 @@
 #include "Editor.h"
 #include "Editor/EditorEngine.h"
 #include "Engine/Blueprint.h"
+#include "Engine/SimpleConstructionScript.h"
 #include "HAL/PlatformCrt.h"
 #include "HAL/PlatformMath.h"
 #include "Internationalization/Text.h"
@@ -553,6 +554,11 @@ void UBlueprintEditorLibrary::ReparentBlueprint(UBlueprint* Blueprint, UClass* N
 
 	Blueprint->ParentClass = NewParentClass;
 
+	if (Blueprint->SimpleConstructionScript != nullptr)
+	{
+		Blueprint->SimpleConstructionScript->ValidateSceneRootNodes();
+	}
+
 	FBlueprintEditorUtils::RefreshAllNodes(Blueprint);
 	FBlueprintEditorUtils::MarkBlueprintAsModified(Blueprint);
 
