Add CreateBlueprintAssetWithParent editor function, this avoids the need to reparent a blueprint after creating it
[RN] Added CreateBlueprintAssetWithParent editor function, for creating blueprint assets with a specific parent class

diff --git a/Engine/Source/Editor/BlueprintEditorLibrary/Private/BlueprintEditorLibrary.cpp b/Engine/Source/Editor/BlueprintEditorLibrary/Private/BlueprintEditorLibrary.cpp
--- a/Engine/Source/Editor/BlueprintEditorLibrary/Private/BlueprintEditorLibrary.cpp
+++ b/Engine/Source/Editor/BlueprintEditorLibrary/Private/BlueprintEditorLibrary.cpp
@@ -2,6 +2,7 @@
 
 #include "BlueprintEditorLibrary.h"
 
+#include "AssetRegistry/AssetRegistryModule.h"
 #include "BlueprintEditorModule.h"
 #include "BlueprintEditor.h"
 #include "AnimGraphNode_Base.h"
@@ -29,9 +30,12 @@
 #include "Kismet2/BlueprintEditorUtils.h"
 #include "Kismet2/Kismet2NameValidators.h"
 #include "Kismet2/KismetEditorUtilities.h"
+#include "KismetCompilerModule.h"
 #include "Logging/LogCategory.h"
+#include "Logging/StructuredLog.h"
 #include "Math/Vector2D.h"
 #include "Misc/AssertionMacros.h"
+#include "PackageTools.h"
 #include "Subsystems/AssetEditorSubsystem.h"
 #include "Templates/Casts.h"
 #include "Templates/SubclassOf.h"
@@ -692,6 +696,79 @@ void UBlueprintEditorLibrary::SetBlueprintVariableExposeOnSpawn(UBlueprint* Blue
 	} 
 }
 
+UBlueprint* UBlueprintEditorLibrary::CreateBlueprintAssetWithParent(const FString& AssetPath, UClass* ParentClass)
+{
+	if(!ParentClass)
+	{
+		UE_LOGFMT(LogBlueprintEditorLib, Warning, "Cannot create a blueprint asset with null parent class");
+		return nullptr;
+	}
+	
+	// do not allow inheritance of function library blueprints or native function libraries that already have functions
+	// bIsValidFunctionLibrary provides a carve out for UEditorFunctionLibrary and similar sentinel types:
+	const bool bIsFunctionLibrary = ParentClass->IsChildOf(UBlueprintFunctionLibrary::StaticClass());
+	const bool bIsValidFunctionLibrary = bIsFunctionLibrary && (ParentClass->Children == nullptr && ParentClass->HasAnyClassFlags(CLASS_Native));
+	if(bIsFunctionLibrary && !bIsValidFunctionLibrary)
+	{
+		UE_LOGFMT(LogBlueprintEditorLib, Warning, "Cannot create a blueprint asset from a function library: {ClassPath}", ParentClass->GetPathName());
+		return nullptr;
+	}
+
+	// Validate base blueprint logic - this enforces 'blueprintable/notblueprintable'
+	if(!bIsValidFunctionLibrary && !FKismetEditorUtilities::CanCreateBlueprintOfClass(ParentClass))
+	{
+		UE_LOGFMT(LogBlueprintEditorLib, Warning, "Not allowed to create blueprint for class: {ClassPath} - is it Blueprintable or IsBlueprintBase?", ParentClass->GetPathName());
+		return nullptr;
+	}
+
+	// interface classes require special handling - reject them:
+	if(ParentClass->HasAnyClassFlags(CLASS_Interface))
+	{
+		UE_LOGFMT(LogBlueprintEditorLib, Warning, "Cannot create a blueprint asset from an interface: {ClassPath}", ParentClass->GetPathName());
+		return nullptr;
+	}
+
+	const FString PackageName = UPackageTools::SanitizePackageName(AssetPath);
+	UPackage* Existing = FindObject<UPackage>(nullptr, *PackageName);
+	if(Existing)
+	{
+		UE_LOGFMT(LogBlueprintEditorLib, Warning, "Cannot create a blueprint asset because an asset with this name already exists: {PackageName}", PackageName);
+		return nullptr;
+	}
+
+	UPackage* Pkg = CreatePackage(*PackageName);
+	if(!Pkg)
+	{	
+		UE_LOGFMT(LogBlueprintEditorLib, Warning, "Create Package Failed: {PackageName}", PackageName);
+		return nullptr;
+	}
+
+	FName BPName = FPackageName::GetShortFName(PackageName);
+	
+	UClass* BlueprintClass = nullptr;
+	UClass* BlueprintGeneratedClass = nullptr;
+	IKismetCompilerInterface& KismetCompilerModule = FModuleManager::LoadModuleChecked<IKismetCompilerInterface>("KismetCompiler");
+	KismetCompilerModule.GetBlueprintTypesForClass(ParentClass, BlueprintClass, BlueprintGeneratedClass);
+	
+	UBlueprint* BP = FKismetEditorUtilities::CreateBlueprint(
+		ParentClass, 
+		Pkg, 
+		BPName, 
+		bIsFunctionLibrary ? BPTYPE_FunctionLibrary : BPTYPE_Normal, 
+		BlueprintClass, 
+		BlueprintGeneratedClass);
+	ensure(BP); // FKismetEditorUtilities::CreateBlueprint does not return null, if it does we should clean up the UPackage - somehow
+
+	Pkg->SetIsExternallyReferenceable(true);
+
+	// Notify the asset registry
+	FAssetRegistryModule::AssetCreated(BP);
+
+	// Mark the package dirty...
+	Pkg->MarkPackageDirty();
+	return BP;
+}
+
 void UBlueprintEditorLibrary::SetBlueprintVariableExposeToCinematics(UBlueprint* Blueprint, const FName& VariableName, bool bExposeToCinematics)
 {
 	if (!Blueprint)
diff --git a/Engine/Source/Editor/BlueprintEditorLibrary/Public/BlueprintEditorLibrary.h b/Engine/Source/Editor/BlueprintEditorLibrary/Public/BlueprintEditorLibrary.h
--- a/Engine/Source/Editor/BlueprintEditorLibrary/Public/BlueprintEditorLibrary.h
+++ b/Engine/Source/Editor/BlueprintEditorLibrary/Public/BlueprintEditorLibrary.h
@@ -215,4 +215,13 @@ public:
 	 */
 	UFUNCTION(BlueprintCallable, Category = "Editor Scripting", meta = (ScriptMethod))
 	static void SetBlueprintVariableInstanceEditable(UBlueprint* Blueprint, const FName& VariableName, bool bInstanceEditable);
+
+	/**
+	 * Creates a blueprint based on a specific parent, honoring registered custom blueprint types
+	 * 
+	 * @param AssetPath				The full path that the asset should be created with
+	 * @param ParentClass			The parent class that the blueprint should be based on
+	 */
+	 UFUNCTION(BlueprintCallable, Category = "Blueprint Editor")
+	 static UBlueprint* CreateBlueprintAssetWithParent(const FString& AssetPath, UClass* ParentClass);
 };
\ No newline at end of file
