Avoid crash when attempting to get generated class from a blueprint that has not yet generated a class

#rnx

diff --git a/Engine/Source/Editor/BlueprintEditorLibrary/Private/BlueprintEditorLibrary.cpp b/Engine/Source/Editor/BlueprintEditorLibrary/Private/BlueprintEditorLibrary.cpp
--- a/Engine/Source/Editor/BlueprintEditorLibrary/Private/BlueprintEditorLibrary.cpp
+++ b/Engine/Source/Editor/BlueprintEditorLibrary/Private/BlueprintEditorLibrary.cpp
@@ -667,6 +667,11 @@ UClass* UBlueprintEditorLibrary::GeneratedClass(UBlueprint* BlueprintObj)
 {
 	if (BlueprintObj)
 	{
+		if(BlueprintObj->GeneratedClass == nullptr)
+		{
+			UE_LOG(LogBlueprintEditorLib, Warning, TEXT("Blueprint %s does not have a generated class - consider compiling it"), *BlueprintObj->GetPathName());
+			return nullptr;
+		}
 		return BlueprintObj->GeneratedClass->GetAuthoritativeClass();
 	}
 	return nullptr;
