Expose member variable creation to editor blueprints and python editor scripts

[RN] Editor scripts can now add variables to blueprint assets using AddMemberVariable

diff --git a/Engine/Source/Editor/BlueprintEditorLibrary/Private/BlueprintEditorLibrary.cpp b/Engine/Source/Editor/BlueprintEditorLibrary/Private/BlueprintEditorLibrary.cpp
--- a/Engine/Source/Editor/BlueprintEditorLibrary/Private/BlueprintEditorLibrary.cpp
+++ b/Engine/Source/Editor/BlueprintEditorLibrary/Private/BlueprintEditorLibrary.cpp
@@ -3,6 +3,7 @@
 #include "BlueprintEditorLibrary.h"
 
 #include "AssetRegistry/AssetRegistryModule.h"
+#include "Blueprint/BlueprintExceptionInfo.h"
 #include "BlueprintEditorModule.h"
 #include "BlueprintEditor.h"
 #include "AnimGraphNode_Base.h"
@@ -807,4 +808,180 @@ void UBlueprintEditorLibrary::SetBlueprintVariableInstanceEditable(UBlueprint* B
 	
 	FBlueprintEditorUtils::SetBlueprintOnlyEditableFlag(Blueprint, VariableName, !bInstanceEditable);
 }
-#undef LOCTEXT_NAMESPACE	// "BlueprintEditorLibrary"
+
+bool UBlueprintEditorLibrary::Generic_AddMemberVariableWithValue(UBlueprint* Blueprint, FName MemberName, const uint8* DefaultValuePtr, const FProperty* DefaultValueProp)
+{
+	const UEdGraphSchema_K2* Schema = GetDefault<UEdGraphSchema_K2>();
+	FEdGraphPinType PinType;
+	if( !Schema->ConvertPropertyToPinType(DefaultValueProp, PinType) )
+	{
+		return false;
+	}
+
+	FName VarName = FBlueprintEditorUtils::FindUniqueKismetName(Blueprint, MemberName.ToString(), Blueprint->SkeletonGeneratedClass);
+	return FBlueprintEditorUtils::AddMemberVariable(Blueprint, VarName, PinType);
+}
+
+DEFINE_FUNCTION(UBlueprintEditorLibrary::execAddMemberVariableWithValue)
+{
+	P_GET_OBJECT(UBlueprint, Blueprint);
+	P_GET_PROPERTY(FNameProperty, MemberName);
+	
+	Stack.StepCompiledIn<FProperty>(nullptr);
+	const FProperty* DefaultValueProp = Stack.MostRecentProperty;
+	const uint8* DefaultValuePtr = Stack.MostRecentPropertyAddress;
+	P_FINISH;
+
+	if (!DefaultValueProp || !DefaultValuePtr)
+	{
+		FBlueprintExceptionInfo ExceptionInfo(
+			EBlueprintExceptionType::AccessViolation,
+			LOCTEXT("AddMemberVariable_MissingValue", "Failed to resolve default value and property type from AddMemberVariable.")
+		);
+		FBlueprintCoreDelegates::ThrowScriptException(P_THIS, Stack, ExceptionInfo);
+	}
+
+	*(bool*)RESULT_PARAM = P_THIS->Generic_AddMemberVariableWithValue(Blueprint, MemberName, DefaultValuePtr, DefaultValueProp);
+}
+
+bool UBlueprintEditorLibrary::AddMemberVariable(UBlueprint* Blueprint, FName MemberName, const FEdGraphPinType& VariableType)
+{
+	return FBlueprintEditorUtils::AddMemberVariable(
+		Blueprint, 
+		FBlueprintEditorUtils::FindUniqueKismetName(Blueprint, MemberName.ToString(), Blueprint->SkeletonGeneratedClass),
+		VariableType);
+}
+
+FEdGraphPinType UBlueprintEditorLibrary::GetBasicTypeByName(FName TypeName)
+{
+	FEdGraphPinType Result;
+	const TSet<FName> PrimitiveTypes = {
+		UEdGraphSchema_K2::PC_Boolean, 
+		UEdGraphSchema_K2::PC_Byte,
+		UEdGraphSchema_K2::PC_Int,
+		UEdGraphSchema_K2::PC_Int64,
+		UEdGraphSchema_K2::PC_Real,
+		UEdGraphSchema_K2::PC_Name,
+		UEdGraphSchema_K2::PC_String,
+		UEdGraphSchema_K2::PC_Text,
+	};
+	if(!PrimitiveTypes.Contains(TypeName))
+	{
+		UE_LOG(LogBlueprintEditorLib, Warning, TEXT("Primitive type: %s not recognized, defaulting to int"), *TypeName.ToString());
+		TypeName = UEdGraphSchema_K2::PC_Int;
+	}
+	Result.PinCategory = TypeName;
+	if(TypeName == UEdGraphSchema_K2::PC_Real)
+	{
+		Result.PinSubCategory == UEdGraphSchema_K2::PC_Double;
+	}
+	return Result;
+}
+	
+FEdGraphPinType UBlueprintEditorLibrary::GetStructType(const UScriptStruct* StructType)
+{
+	if(	StructType == nullptr || 
+		!UEdGraphSchema_K2::IsAllowableBlueprintVariableType(StructType))
+	{
+		UE_LOG(LogBlueprintEditorLib, Warning, TEXT("Struct type: %s not allowed, defaulting to int"), StructType ? *StructType->GetPathName() : TEXT("null"));
+		return GetBasicTypeByName(UEdGraphSchema_K2::PC_Int);
+	}
+	
+	FEdGraphPinType Result;
+	Result.PinCategory = UEdGraphSchema_K2::PC_Struct;
+	// the struct here is notionally const, and via PinSubCategoryObject should be extremely rare if they exist:
+	Result.PinSubCategoryObject = const_cast<UScriptStruct*>(StructType);
+	return Result;
+}
+	
+FEdGraphPinType UBlueprintEditorLibrary::GetClassReferenceType(const UClass* ClassType)
+{
+	if(	ClassType == nullptr || 
+		!UEdGraphSchema_K2::IsAllowableBlueprintVariableType(ClassType))
+	{
+		UE_LOG(LogBlueprintEditorLib, Warning, TEXT("Class type: %s not allowed, defaulting to int"), ClassType ? *ClassType->GetPathName() : TEXT("null"));
+		return GetBasicTypeByName(UEdGraphSchema_K2::PC_Int);
+	}
+	
+	FEdGraphPinType Result;
+	Result.PinCategory = UEdGraphSchema_K2::PC_Class;
+	// the class here is notionally const, and via PinSubCategoryObject should be extremely rare if they exist:
+	Result.PinSubCategoryObject = const_cast<UClass*>(ClassType);
+	return Result;
+}
+	
+FEdGraphPinType UBlueprintEditorLibrary::GetObjectReferenceType(const UClass* ObjectType)
+{
+	if(	ObjectType == nullptr || 
+		!UEdGraphSchema_K2::IsAllowableBlueprintVariableType(ObjectType))
+	{
+		UE_LOG(LogBlueprintEditorLib, Warning, TEXT("Object reference type: %s not allowed, defaulting to int"), ObjectType ? *ObjectType->GetPathName() : TEXT("null"));
+		return GetBasicTypeByName(UEdGraphSchema_K2::PC_Int);
+	}
+	
+	FEdGraphPinType Result;
+	Result.PinCategory = UEdGraphSchema_K2::PC_Object;
+	// the struct here is notionally const, and via PinSubCategoryObject should be extremely rare if they exist:
+	Result.PinSubCategoryObject = const_cast<UClass*>(ObjectType);
+	return Result;
+}
+	
+FEdGraphPinType UBlueprintEditorLibrary::GetArrayType(const FEdGraphPinType& ContainedType)
+{
+	if(ContainedType.IsContainer())
+	{
+		UE_LOG(LogBlueprintEditorLib, Warning, TEXT("Containers cannot be nested directly, an intermediate struct type must be created. Defaulting to int"));
+		return GetBasicTypeByName(UEdGraphSchema_K2::PC_Int);
+	}
+
+	FEdGraphPinType Result = ContainedType;
+	Result.ContainerType = EPinContainerType::Array;
+	return Result;
+}
+	
+FEdGraphPinType UBlueprintEditorLibrary::GetSetType(const FEdGraphPinType& ContainedType)
+{
+	if(ContainedType.IsContainer())
+	{
+		UE_LOG(LogBlueprintEditorLib, Warning, TEXT("Containers cannot be nested directly, an intermediate struct type must be created. Defaulting to int"));
+		return GetBasicTypeByName(UEdGraphSchema_K2::PC_Int);
+	}
+
+	if(!FBlueprintEditorUtils::HasGetTypeHash(ContainedType))
+	{
+		UE_LOG(LogBlueprintEditorLib, Warning, TEXT("Key type must be hashable. Defaulting to int"));
+		return GetBasicTypeByName(UEdGraphSchema_K2::PC_Int);
+	}
+
+	FEdGraphPinType Result= ContainedType;
+	Result.ContainerType = EPinContainerType::Set;
+	return Result;
+}
+	
+FEdGraphPinType UBlueprintEditorLibrary::GetMapType(const FEdGraphPinType& KeyType,const FEdGraphPinType& ValueType)
+{
+	if(KeyType.IsContainer())
+	{
+		UE_LOG(LogBlueprintEditorLib, Warning, TEXT("Containers cannot be used as a key type, an intermediate struct type must be created. Defaulting to int"));
+		return GetBasicTypeByName(UEdGraphSchema_K2::PC_Int);
+	}
+
+	if(ValueType.IsContainer())
+	{
+		UE_LOG(LogBlueprintEditorLib, Warning, TEXT("Containers cannot be as a value type, an intermediate struct type must be created. Defaulting to int"));
+		return GetBasicTypeByName(UEdGraphSchema_K2::PC_Int);
+	}
+
+	if(!FBlueprintEditorUtils::HasGetTypeHash(KeyType))
+	{
+		UE_LOG(LogBlueprintEditorLib, Warning, TEXT("Key type must be hashable. Defaulting to int"));
+		return GetBasicTypeByName(UEdGraphSchema_K2::PC_Int);
+	}
+
+	FEdGraphPinType Result = KeyType;
+	Result.ContainerType = EPinContainerType::Map;
+	Result.PinValueType = FEdGraphTerminalType::FromPinType(ValueType);
+	return Result;
+}
+
+#undef LOCTEXT_NAMESPACE	// "BlueprintEditorLibrary"
\ No newline at end of file
diff --git a/Engine/Source/Editor/BlueprintEditorLibrary/Public/BlueprintEditorLibrary.h b/Engine/Source/Editor/BlueprintEditorLibrary/Public/BlueprintEditorLibrary.h
--- a/Engine/Source/Editor/BlueprintEditorLibrary/Public/BlueprintEditorLibrary.h
+++ b/Engine/Source/Editor/BlueprintEditorLibrary/Public/BlueprintEditorLibrary.h
@@ -224,4 +224,50 @@ public:
 	 */
 	 UFUNCTION(BlueprintCallable, Category = "Blueprint Editor")
 	 static UBlueprint* CreateBlueprintAssetWithParent(const FString& AssetPath, UClass* ParentClass);
+
+	/**
+	  * Adds a member variable to the specified blueprint inferring the type from a provided value.
+	  * 
+	  * @return	true if it succeeds, false if it fails.
+	  */
+	UFUNCTION(BlueprintCallable, CustomThunk, Category = "Blueprint Editor", meta=(CustomStructureParam="DefaultValue"))
+	static bool AddMemberVariableWithValue(UBlueprint* Blueprint, FName MemberName, const int32& DefaultValue);
+	static bool Generic_AddMemberVariableWithValue(UBlueprint* Blueprint, FName MemberName, const uint8* DefaultValuePtr, const FProperty* DefaultValueProp);
+	DECLARE_FUNCTION(execAddMemberVariableWithValue);
+	
+	/**
+	  * Adds a member variable to the specified blueprint with the specified type.
+	  * 
+	  * @return	true if it succeeds, false if it fails.
+	  */
+	UFUNCTION(BlueprintCallable, Category = "Blueprint Editor")
+	static bool AddMemberVariable(UBlueprint* Blueprint, FName MemberName, const FEdGraphPinType& VariableType);
+	
+	/** @return a pintype for 'int', 'byte', 'bool', 'real', 'name', 'string' or 'text' - returns 'int' type if invalid type is provided */
+	UFUNCTION(BlueprintCallable, Category = "Blueprint Editor")
+	static FEdGraphPinType GetBasicTypeByName(FName TypeName);
+	
+	/** @return a pintype for the provided struct - returns 'int' type if invalid struct is provided */
+	UFUNCTION(BlueprintCallable, Category = "Blueprint Editor")
+	static FEdGraphPinType GetStructType(const UScriptStruct* StructType);
+	
+	/** @return a class reference pintype for the provided class - returns 'int' type if invalid class is provided */
+	UFUNCTION(BlueprintCallable, Category = "Blueprint Editor")
+	static FEdGraphPinType GetClassReferenceType(const UClass* ClassType);
+	
+	/** @return a object reference pintype for the provided class - returns 'int' type if invalid object type is provided */
+	UFUNCTION(BlueprintCallable, Category = "Blueprint Editor")
+	static FEdGraphPinType GetObjectReferenceType(const UClass* ObjectType);
+	
+	/** @return a array of ContainedType type - returns 'int' type if invalid type is provided */
+	UFUNCTION(BlueprintCallable, Category = "Blueprint Editor")
+	static FEdGraphPinType GetArrayType(const FEdGraphPinType& ContainedType);
+	
+	/** @return a set of ContainedType type - returns 'int' type if invalid type is provided */
+	UFUNCTION(BlueprintCallable, Category = "Blueprint Editor")
+	static FEdGraphPinType GetSetType(const FEdGraphPinType& ContainedType);
+	
+	/** @return a map of KeyType to ValueType type - returns 'int' type if invalid type is provided */
+	UFUNCTION(BlueprintCallable, Category = "Blueprint Editor")
+	static FEdGraphPinType GetMapType(const FEdGraphPinType& KeyType,const FEdGraphPinType& ValueType);
 };
\ No newline at end of file
