Fix error reintroduced from shelf

#rnx

diff --git a/Engine/Source/Editor/BlueprintEditorLibrary/Private/BlueprintEditorLibrary.cpp b/Engine/Source/Editor/BlueprintEditorLibrary/Private/BlueprintEditorLibrary.cpp
--- a/Engine/Source/Editor/BlueprintEditorLibrary/Private/BlueprintEditorLibrary.cpp
+++ b/Engine/Source/Editor/BlueprintEditorLibrary/Private/BlueprintEditorLibrary.cpp
@@ -873,7 +873,7 @@ FEdGraphPinType UBlueprintEditorLibrary::GetBasicTypeByName(FName TypeName)
 	Result.PinCategory = TypeName;
 	if(TypeName == UEdGraphSchema_K2::PC_Real)
 	{
-		Result.PinSubCategory == UEdGraphSchema_K2::PC_Double;
+		Result.PinSubCategory = UEdGraphSchema_K2::PC_Double;
 	}
 	return Result;
 }
