Avoid crash when a user provides no blueprint to the new add member variable editor functions

#rnx

diff --git a/Engine/Source/Editor/BlueprintEditorLibrary/Private/BlueprintEditorLibrary.cpp b/Engine/Source/Editor/BlueprintEditorLibrary/Private/BlueprintEditorLibrary.cpp
--- a/Engine/Source/Editor/BlueprintEditorLibrary/Private/BlueprintEditorLibrary.cpp
+++ b/Engine/Source/Editor/BlueprintEditorLibrary/Private/BlueprintEditorLibrary.cpp
@@ -811,6 +811,12 @@ void UBlueprintEditorLibrary::SetBlueprintVariableInstanceEditable(UBlueprint* B
 
 bool UBlueprintEditorLibrary::Generic_AddMemberVariableWithValue(UBlueprint* Blueprint, FName MemberName, const uint8* DefaultValuePtr, const FProperty* DefaultValueProp)
 {
+	if(!Blueprint)
+	{
+		UE_LOG(LogBlueprintEditorLib, Warning, TEXT("Attempted to add member variable to null blueprint"));
+		return false;	
+	}
+
 	const UEdGraphSchema_K2* Schema = GetDefault<UEdGraphSchema_K2>();
 	FEdGraphPinType PinType;
 	if( !Schema->ConvertPropertyToPinType(DefaultValueProp, PinType) )
@@ -846,6 +852,12 @@ DEFINE_FUNCTION(UBlueprintEditorLibrary::execAddMemberVariableWithValue)
 
 bool UBlueprintEditorLibrary::AddMemberVariable(UBlueprint* Blueprint, FName MemberName, const FEdGraphPinType& VariableType)
 {
+	if(!Blueprint)
+	{
+		UE_LOG(LogBlueprintEditorLib, Warning, TEXT("Attempted to add member variable to null blueprint"));
+		return false;	
+	}
+
 	return FBlueprintEditorUtils::AddMemberVariable(
 		Blueprint, 
 		FBlueprintEditorUtils::FindUniqueKismetName(Blueprint, MemberName.ToString(), Blueprint->SkeletonGeneratedClass),
