Add function for getting the blueprint associated with a class, if any. Because most classes do not have a blueprint, this function encourages callers to branch on success

[RN] Added GetBlueprintForClass to find the Blueprint object that generated a class, if any

diff --git a/Engine/Source/Editor/BlueprintEditorLibrary/Private/BlueprintEditorLibrary.cpp b/Engine/Source/Editor/BlueprintEditorLibrary/Private/BlueprintEditorLibrary.cpp
--- a/Engine/Source/Editor/BlueprintEditorLibrary/Private/BlueprintEditorLibrary.cpp
+++ b/Engine/Source/Editor/BlueprintEditorLibrary/Private/BlueprintEditorLibrary.cpp
@@ -512,6 +512,22 @@ UBlueprint* UBlueprintEditorLibrary::GetBlueprintAsset(UObject* Object)
 	return Cast<UBlueprint>(Object);
 }
 
+UBlueprint* UBlueprintEditorLibrary::GetBlueprintForClass(UClass* Class, bool& bDoesClassHaveBlueprint)
+{
+	bDoesClassHaveBlueprint = false;
+	if(!Class)
+	{
+		return nullptr;
+	}
+
+	if(UBlueprint* Result = Cast<UBlueprint>(Class->ClassGeneratedBy))
+	{
+		bDoesClassHaveBlueprint = true;
+		return Result;
+	}
+	return nullptr;
+}
+
 void UBlueprintEditorLibrary::RefreshOpenEditorsForBlueprint(const UBlueprint* BP)
 {
 	// Get any open blueprint editors for this asset and refresh them if they match the given blueprint
diff --git a/Engine/Source/Editor/BlueprintEditorLibrary/Public/BlueprintEditorLibrary.h b/Engine/Source/Editor/BlueprintEditorLibrary/Public/BlueprintEditorLibrary.h
--- a/Engine/Source/Editor/BlueprintEditorLibrary/Public/BlueprintEditorLibrary.h
+++ b/Engine/Source/Editor/BlueprintEditorLibrary/Public/BlueprintEditorLibrary.h
@@ -137,6 +137,19 @@ public:
 	*/
 	UFUNCTION(BlueprintCallable, Category = "Blueprint Upgrade Tools", meta = (Keywords = "cast"))
 	static UBlueprint* GetBlueprintAsset(UObject* Object);
+	
+	/**
+	 * Looks up the UBlueprint that generated the provided class, if any. Provides a 'true' exec pin
+	 * to execute if there is a valid blueprint associated with the Class.
+	 * 
+	 * @param Class						The class to look up the blueprint for
+	 * @param bDoesClassHaveBlueprint	Whether the provided class had a blueprint
+	 * 
+	 * @return							The blueprint that generated the class, nullptr if the UClass
+										is native or otherwise cooked
+	 */
+	UFUNCTION(BlueprintCallable, Category = "Blueprint Upgrade Tools", meta = (ExpandBoolAsExecs = "bDoesClassHaveBlueprint"))
+	static UBlueprint* GetBlueprintForClass(UClass* Class, bool& bDoesClassHaveBlueprint);
 
 	/** Attempt to refresh any open blueprint editors for the given asset */
 	UFUNCTION(BlueprintCallable, Category = "Blueprint Editor")
