Add some Blueprint Editor Utilitiy functions to compare an assets saved engine version.

#wip

diff --git a/Engine/Source/Editor/BlueprintEditorLibrary/Private/BlueprintEditorLibrary.cpp b/Engine/Source/Editor/BlueprintEditorLibrary/Private/BlueprintEditorLibrary.cpp
--- a/Engine/Source/Editor/BlueprintEditorLibrary/Private/BlueprintEditorLibrary.cpp
+++ b/Engine/Source/Editor/BlueprintEditorLibrary/Private/BlueprintEditorLibrary.cpp
@@ -22,6 +22,7 @@
 #include "Editor/EditorEngine.h"
 #include "Engine/Blueprint.h"
 #include "Engine/SimpleConstructionScript.h"
+#include "HAL/FileManager.h"
 #include "HAL/PlatformCrt.h"
 #include "HAL/PlatformMath.h"
 #include "Internationalization/Text.h"
@@ -36,6 +37,7 @@
 #include "Logging/StructuredLog.h"
 #include "Math/Vector2D.h"
 #include "Misc/AssertionMacros.h"
+#include "Misc/EngineVersion.h"
 #include "PackageTools.h"
 #include "Subsystems/AssetEditorSubsystem.h"
 #include "Templates/Casts.h"
@@ -139,6 +141,29 @@ namespace InternalBlueprintEditorLibrary
 			AssetSubsystem->CloseAllEditorsForAsset(Blueprint);
 		}
 	}
+
+	/**
+	 * Converts the given EVersionComparison to the BP friendly EAssetSaveVersionComparisonResults type.
+	 */
+	static EAssetSaveVersionComparisonResults ConvertSaveVersionComparison(const EVersionComparison RawVersionCompare)
+	{
+		EAssetSaveVersionComparisonResults Result = EAssetSaveVersionComparisonResults::InvalidComparison;
+	
+		switch (RawVersionCompare)
+		{
+		case EVersionComparison::Neither:
+			Result = EAssetSaveVersionComparisonResults::Identical;
+			break;
+		case EVersionComparison::First:
+			Result = EAssetSaveVersionComparisonResults::Newer;
+			break;
+		case EVersionComparison::Second:
+			Result = EAssetSaveVersionComparisonResults::Older;
+			break;
+		}
+
+		return Result;
+	}
 };
 
 ///////////////////////////////////////////////////////////
@@ -164,6 +189,125 @@ UEdGraph* UBlueprintEditorLibrary::FindEventGraph(UBlueprint* Blueprint)
 	return Blueprint ? FBlueprintEditorUtils::FindEventGraph(Blueprint) : nullptr;
 }
 
+void UBlueprintEditorLibrary::CompareAssetSaveVersionTo(const UObject* Asset, const FString& VersionToCheckString, EAssetSaveVersionComparisonResults& Result)
+{
+	Result = EAssetSaveVersionComparisonResults::InvalidComparison;
+	
+	if (!Asset)
+	{
+		UE_LOG(LogBlueprintEditorLib, Error, TEXT("[%hs] 'Asset' is null! Cannot compare to engine version '%s'"), __func__, *VersionToCheckString);
+		return;
+	}
+
+	FEngineVersion VersionToCheck = {};
+	const bool bSuccessfulParse = FEngineVersion::Parse(VersionToCheckString, OUT VersionToCheck);
+	if (!bSuccessfulParse)
+	{
+		UE_LOG(LogBlueprintEditorLib, Error, TEXT("[%hs] 'VersionToCheckString' value of '%s' is not a valid FEngineVersion!"), __func__, *VersionToCheckString);
+		return;
+	}
+	
+	// The linker has the data about what engine version was used to save this asset
+	FLinkerLoad* Linker = Asset->GetLinker();
+	if (!Linker)
+	{
+		UE_LOG(LogBlueprintEditorLib, Error, TEXT("[%hs] Failed to find the linker for asset '%s'"), __func__, *GetNameSafe(Asset));
+		return;
+	}
+
+	const FEngineVersion& AssetVersion = Linker->Summary.SavedByEngineVersion;
+	
+	EVersionComponent* DifferingComponent =  nullptr;
+	const EVersionComparison Comparison = FEngineVersionBase::GetNewest(AssetVersion, VersionToCheck, DifferingComponent);
+
+	Result = InternalBlueprintEditorLibrary::ConvertSaveVersionComparison(Comparison);
+}
+
+void UBlueprintEditorLibrary::CompareSoftObjectSaveVersionTo(const TSoftObjectPtr<UObject> ObjectToCheck, const FString& VersionToCheckString, EAssetSaveVersionComparisonResults& Result)
+{
+	Result = EAssetSaveVersionComparisonResults::InvalidComparison;
+
+	if (!ObjectToCheck.IsValid())
+	{
+		UE_LOG(LogBlueprintEditorLib, Error, TEXT("[%hs] An invalid ObjectToCheck has been provided, cannot compare save versions."), __func__);
+		return;
+	}
+	
+	FString AbsolutePackageFilePath;
+	const bool bSuccessfulyFoundFile = FPackageName::DoesPackageExist(ObjectToCheck.GetLongPackageName(), OUT &AbsolutePackageFilePath);
+	if (!bSuccessfulyFoundFile)
+	{
+		UE_LOG(LogBlueprintEditorLib, Error, TEXT("[%hs] Failed to find package on disk for soft object '%s'"), __func__, *ObjectToCheck.ToString());
+		return;
+	}
+	
+	// Ensure that this is indeed a package file path. This should always be true if the above DoesPackageExist function works.
+	if (!ensure(FPackageName::IsPackageFilename(AbsolutePackageFilePath)))
+	{
+		UE_LOG(LogBlueprintEditorLib, Error, TEXT("[%hs] '%s' is not a package file path! (FPackageName::IsPackageFilename returned false)"), __func__, *AbsolutePackageFilePath);
+		return;
+	}
+
+	// Make sure we have a valid version to compare to before attempting to open a file reader
+	FEngineVersion VersionToCheck = {};
+	const bool bSuccessfulParse = FEngineVersion::Parse(VersionToCheckString, OUT VersionToCheck);
+	if (!bSuccessfulParse)
+	{
+		UE_LOG(LogBlueprintEditorLib, Error, TEXT("[%hs] 'VersionToCheckString' value of '%s' is not a valid FEngineVersion!"), __func__, *VersionToCheckString);
+		return;
+	}
+	
+	// Create a file reader to load the file and read its package summary data, which has the save version
+	TUniquePtr<FArchive> FileReader(IFileManager::Get().CreateFileReader(*AbsolutePackageFilePath));
+	if (!FileReader.IsValid())
+	{
+		UE_LOG(LogBlueprintEditorLib, Error, TEXT("[%hs] Failed to open file reader for path '%s'"), __func__, *AbsolutePackageFilePath);
+		return;
+	}
+	
+	FPackageFileSummary FileSummary;
+	(*FileReader) << FileSummary;
+
+	// Make sure this is indeed a package
+	if (FileSummary.Tag == PACKAGE_FILE_TAG)
+	{
+		EVersionComponent* DifferingComponent =  nullptr;
+		const EVersionComparison Comparison = FEngineVersionBase::GetNewest(FileSummary.SavedByEngineVersion, VersionToCheck, DifferingComponent);
+		Result = InternalBlueprintEditorLibrary::ConvertSaveVersionComparison(Comparison);
+	}
+
+	// Clean up our file reader, we are done with it
+	FileReader->Close();
+	FileReader.Reset();
+}
+
+FString UBlueprintEditorLibrary::GetSavedByEngineVersion(const UObject* Asset)
+{
+	static const FString InvalidVersion = TEXT("INVALID");
+	if (!Asset)
+	{
+		UE_LOG(LogBlueprintEditorLib, Error, TEXT("[%hs] 'Asset' is null! Cannot get the 'saved by' engine version"), __func__);
+		return InvalidVersion;
+	}
+
+	// The linker has the data about what engine version was used to save this asset
+	FLinkerLoad* Linker = Asset->GetLinker();
+	if (!Linker)
+	{
+		UE_LOG(LogBlueprintEditorLib, Error, TEXT("[%hs] Failed to find the linker for asset '%s'"), __func__, *GetNameSafe(Asset));
+		return InvalidVersion;
+	}
+	
+	const FEngineVersion& AssetSavedVersion = Linker->Summary.SavedByEngineVersion;
+	
+	return AssetSavedVersion.ToString();
+}
+
+FString UBlueprintEditorLibrary::GetCurrentEngineVersion()
+{
+	return FEngineVersion::Current().ToString();
+}
+
 UEdGraph* UBlueprintEditorLibrary::FindGraph(UBlueprint* Blueprint, FName GraphName)
 {
 	if (Blueprint && !GraphName.IsNone())
diff --git a/Engine/Source/Editor/BlueprintEditorLibrary/Public/BlueprintEditorLibrary.h b/Engine/Source/Editor/BlueprintEditorLibrary/Public/BlueprintEditorLibrary.h
--- a/Engine/Source/Editor/BlueprintEditorLibrary/Public/BlueprintEditorLibrary.h
+++ b/Engine/Source/Editor/BlueprintEditorLibrary/Public/BlueprintEditorLibrary.h
@@ -25,6 +25,22 @@ struct FFrame;
 
 DECLARE_LOG_CATEGORY_EXTERN(LogBlueprintEditorLib, Warning, All);
 
+/**
+ * The results of comparing an assets save version to another
+ */
+UENUM(BlueprintType)
+enum class EAssetSaveVersionComparisonResults : uint8
+{
+	// The comparison could not be completed
+	InvalidComparison,
+	// The asset save version is identical to what it is being compared to
+	Identical,
+	// The asset save version is newer than what it is being compared to
+	Newer,
+	// The asset save version is older than what it is being compared to
+	Older
+};
+
 UCLASS()
 class BLUEPRINTEDITORLIBRARY_API UBlueprintEditorLibrary : public UBlueprintFunctionLibrary
 {
@@ -52,6 +68,54 @@ public:
 	UFUNCTION(BlueprintCallable, Category = "Blueprint Upgrade Tools")
 	static UEdGraph* FindEventGraph(UBlueprint* Blueprint);
 
+	/**
+	 * Compares the given assets save version to the VersionToCheck. 
+	 * 
+	 * @param Asset				The asset which you would like to check the SavedByEngineVersion of.
+	 * 
+	 * @param VersionToCheck	String representation of the engine version to compare against. For example, "5.6.0-37518009+++UE5+Main"
+	 *							@see GetSavedByEngineVersion and GetCurrentEngineVersion
+	 * 
+	 * @param Result			The outcome of the version comparison
+	 */
+	UFUNCTION(BlueprintCallable, Category = "Blueprint Upgrade Tools",  Meta = (ExpandEnumAsExecs = "Result"))
+	static void CompareAssetSaveVersionTo(const UObject* Asset, const FString& VersionToCheck, EAssetSaveVersionComparisonResults& Result);
+
+	/**
+	 * Compares the given soft object's save version to the VersionToCheck. This will read the packages file header
+	 * 
+	 * @param ObjectToCheck		Soft object pointer to the object whose save version you would like to compare.
+	 * 
+	 * @param VersionToCheck	String representation of the engine version to compare against.  For example, "5.6.0-37518009+++UE5+Main"
+	 *							@see GetSavedByEngineVersion and GetCurrentEngineVersion
+	 * 
+	 * @param Result			The outcome of the version comparison
+	 */
+	UFUNCTION(BlueprintCallable, Category = "Blueprint Upgrade Tools",  Meta = (ExpandEnumAsExecs = "Result"))
+	static void CompareSoftObjectSaveVersionTo(const TSoftObjectPtr<UObject> ObjectToCheck,  const FString& VersionToCheck, EAssetSaveVersionComparisonResults& Result);
+
+	/**
+	 * Returns a string representation of the engine version which the given asset was saved with.
+	 * 
+	 * @see FLinker::Summary::SavedByEngineVersion
+	 * @see FPackageFileSummary
+	 * 
+	 * @param Asset The asset to check the saved by engine version of.
+	 * 
+	 * @return	String representation of the engine version which this asset was saved with. "INVALID" if none. 
+	 *			For example: "5.6.0-37518009+++UE5+Main"
+	 */
+	UFUNCTION(BlueprintCallable, BlueprintPure, Category = "Blueprint Upgrade Tools")
+	static FString GetSavedByEngineVersion(const UObject* Asset);
+
+	/**
+	 * Returns a string which represents the current engine version (FEngineVersion::Current())
+	 * 
+	 * For example: "5.6.0-37518009+++UE5+Main"
+	 */
+	UFUNCTION(BlueprintCallable, BlueprintPure, Category = "Blueprint Upgrade Tools")
+	static FString GetCurrentEngineVersion();
+
 	/**
 	* Finds the graph with the given name on the blueprint. Null if it doesn't have one. 
 	*
