Include the default value in the new member variable

[RN] The BlueprintEditorLibrary AddMemberVariableWithValue function now properly sets the provided default on the created member variable

diff --git a/Engine/Source/Editor/BlueprintEditorLibrary/Private/BlueprintEditorLibrary.cpp b/Engine/Source/Editor/BlueprintEditorLibrary/Private/BlueprintEditorLibrary.cpp
--- a/Engine/Source/Editor/BlueprintEditorLibrary/Private/BlueprintEditorLibrary.cpp
+++ b/Engine/Source/Editor/BlueprintEditorLibrary/Private/BlueprintEditorLibrary.cpp
@@ -984,8 +984,17 @@ bool UBlueprintEditorLibrary::Generic_AddMemberVariableWithValue(UBlueprint* Blu
 		return false;
 	}
 
+	FString ValueAsString;
+	bool bGotDefaultValue = DefaultValueProp->ExportText_Direct(ValueAsString, DefaultValuePtr, DefaultValuePtr, nullptr, PPF_None);
+	if(!bGotDefaultValue)
+	{
+		UE_LOGFMT(LogBlueprintEditorLib, Warning, 
+			"Could not export the provided default value, variable {0} added to {1} will have incorrect default", 
+			MemberName, Blueprint->GetPathName());
+	}
+
 	FName VarName = FBlueprintEditorUtils::FindUniqueKismetName(Blueprint, MemberName.ToString(), Blueprint->SkeletonGeneratedClass);
-	return FBlueprintEditorUtils::AddMemberVariable(Blueprint, VarName, PinType);
+	return FBlueprintEditorUtils::AddMemberVariable(Blueprint, VarName, PinType, ValueAsString);
 }
 
 DEFINE_FUNCTION(UBlueprintEditorLibrary::execAddMemberVariableWithValue)
