Blueprint graph editing and introspection API for any language capable of understanding UFUNCTIONs

#Synced-CL 49751638

diff --git a/Engine/Source/Editor/BlueprintEditorLibrary/BlueprintEditorLibrary.Build.cs b/Engine/Source/Editor/BlueprintEditorLibrary/BlueprintEditorLibrary.Build.cs
--- a/Engine/Source/Editor/BlueprintEditorLibrary/BlueprintEditorLibrary.Build.cs
+++ b/Engine/Source/Editor/BlueprintEditorLibrary/BlueprintEditorLibrary.Build.cs
@@ -19,10 +19,13 @@ public class BlueprintEditorLibrary : ModuleRules
 		PrivateDependencyModuleNames.AddRange(
 			new string[]
 			{
+				"AnimGraph",
 				"CoreUObject",
 				"Engine",
+				"Json",
+				"JsonUtilities",
+				"KismetCompiler",
 				"UnrealEd",
-				"AnimGraph"
 			}
 		);
 	}
diff --git a/Engine/Source/Editor/BlueprintEditorLibrary/Private/BlueprintEditorLibrary.cpp b/Engine/Source/Editor/BlueprintEditorLibrary/Private/BlueprintEditorLibrary.cpp
--- a/Engine/Source/Editor/BlueprintEditorLibrary/Private/BlueprintEditorLibrary.cpp
+++ b/Engine/Source/Editor/BlueprintEditorLibrary/Private/BlueprintEditorLibrary.cpp
@@ -4,6 +4,7 @@
 
 #include "AssetRegistry/AssetRegistryModule.h"
 #include "Blueprint/BlueprintExceptionInfo.h"
+#include "BlueprintEditorLibrary/BlueprintGraphPin.h"
 #include "BlueprintEditorModule.h"
 #include "BlueprintEditor.h"
 #include "AnimGraphNode_Base.h"
@@ -25,9 +26,11 @@
 #include "HAL/FileManager.h"
 #include "HAL/PlatformCrt.h"
 #include "HAL/PlatformMath.h"
+#include "JsonSchema/JsonSchemaGenerator.h"
 #include "Internationalization/Text.h"
 #include "K2Node_CallFunction.h"
 #include "K2Node_CommutativeAssociativeBinaryOperator.h"
+#include "K2Node_IfThenElse.h"
 #include "K2Node_PromotableOperator.h"
 #include "Kismet2/BlueprintEditorUtils.h"
 #include "Kismet2/Kismet2NameValidators.h"
@@ -937,6 +940,26 @@ UBlueprint* UBlueprintEditorLibrary::CreateBlueprintAssetWithParent(const FStrin
 	return BP;
 }
 
+void UBlueprintEditorLibrary::ListGraphNames(const UBlueprint* Blueprint, TArray<FName>& OutGraphNames)
+{
+	if(!IsValid(Blueprint))
+	{
+		UE_LOGFMT(LogBlueprintEditorLib, Warning, "Invalid blueprint has no graphs");
+		return;
+	}
+
+	TArray<UEdGraph*> Graphs;
+	Blueprint->GetAllGraphs(Graphs);
+	for(UEdGraph* Graph : Graphs)
+	{
+		if(!IsValid(Graph) || Graph->HasAnyFlags(RF_Transient))
+		{
+			continue;
+		}
+		OutGraphNames.Add(Graph->GetFName());
+	}
+}
+
 void UBlueprintEditorLibrary::SetBlueprintVariableExposeToCinematics(UBlueprint* Blueprint, const FName& VariableName, bool bExposeToCinematics)
 {
 	if (!Blueprint)
@@ -1167,4 +1190,456 @@ FEdGraphPinType UBlueprintEditorLibrary::GetMapType(const FEdGraphPinType& KeyTy
 	return Result;
 }
 
+FString UBlueprintEditorLibrary::PinTypeToJsonSchema(const FEdGraphPinType& PinType, const UClass* SelfContext)
+{
+	if(PinType.PinCategory == UEdGraphSchema_K2::PC_Exec)
+	{
+		return TEXT("{\n    \"type\": \"exec\"\n}");
+	}
+
+	FCompilerResultsLog Discard;
+	// we could use a scratch object or cache to reduce object graph pollution if needed,
+	// but I want to reuse FPropertyToJsonSchemaObject rather than maintain a direct mapping
+	UStruct* Scope = NewObject<UScriptStruct>();
+	const FProperty* Property = FKismetCompilerUtilities::CreatePropertyOnScope(
+		Scope, 
+		TEXT("TestProperty"), 
+		PinType, 
+		const_cast<UClass*>(SelfContext),
+		CPF_BlueprintVisible|CPF_BlueprintReadOnly, 
+		GetDefault<UEdGraphSchema_K2>(), 
+		Discard);
+	if(!Property)
+	{
+		return TEXT("");
+	}
+
+PRAGMA_DISABLE_DEPRECATION_WARNINGS
+	TSharedPtr<FJsonObject> PinSchema = FJsonSchemaGenerator::FPropertyToJsonSchemaObject(Property);
+PRAGMA_ENABLE_DEPRECATION_WARNINGS	
+	FString OutString;
+	if(PinSchema)
+	{
+		TSharedRef<TJsonWriter<>> JsonWriter = TJsonWriterFactory<>::Create(&OutString);
+		FJsonSerializer::Serialize(PinSchema, JsonWriter);
+		JsonWriter->Close();	
+	}
+	return OutString;
+}
+
+namespace UE::Private
+{
+static void ForEachVisiblePin(const UK2Node* Node, TFunctionRef<bool(UEdGraphPin*)> Ftor)
+{
+	for(UEdGraphPin* Pin : Node->Pins)
+	{
+		if(Pin->bHidden)
+		{
+			continue;
+		}
+
+		if(!Ftor(Pin))
+		{
+			return;
+		}
+	}
+}
+
+static void ForEachVisibleInputPin(const UK2Node* Node, TFunctionRef<bool(UEdGraphPin*)> Ftor)
+{
+	ForEachVisiblePin(Node, [&Ftor](UEdGraphPin* Pin)
+	{
+		if(Pin->Direction == EGPD_Input)
+		{
+			return Ftor(Pin);
+		}
+		return true;
+	});
+}
+
+static void ForEachVisibleOutputPin(const UK2Node* Node, TFunctionRef<bool(UEdGraphPin*)> Ftor)
+{
+	ForEachVisiblePin(Node, [&Ftor](UEdGraphPin* Pin)
+	{
+		if(Pin->Direction == EGPD_Output)
+		{
+			return Ftor(Pin);
+		}
+		return true;
+	});
+}
+
+static UEdGraphPin* GetIfTypeMatches(const UClass* Class, UEdGraphPin* Pin, const FEdGraphPinType& Type)
+{
+	const UEdGraphSchema_K2* K2Schema = GetDefault<UEdGraphSchema_K2>();
+	if(Type == FEdGraphPinType()) // unspecified default type, for this internal helper we treat that as a match
+	{
+		return Pin;
+	}
+	else if (K2Schema->ArePinTypesCompatible(Type, Pin->PinType, Class))
+	{
+		return Pin;
+	}
+	return nullptr;
+}
+}
+
+TArray<UBlueprintGraphPin*> UBlueprintEditorLibrary::ListAllPins(const UK2Node* Node, TEnumAsByte<EEdGraphPinDirection> InDirection)
+{
+	if(!::IsValid(Node))
+	{
+		return {};
+	}
+
+	TArray<UBlueprintGraphPin*> Result;
+	// first input pins - stable order will be more useful to callers
+	if(InDirection == EGPD_MAX || InDirection == EGPD_Input)
+	{
+		UE::Private::ForEachVisibleInputPin(Node, [&Result](UEdGraphPin* Pin)
+		{
+			Result.Add(UBlueprintGraphPin::FromNativePin(Pin));
+			return true;
+		});
+	}
+	// then output pins
+	if(InDirection == EGPD_MAX || InDirection == EGPD_Output)
+	{
+		UE::Private::ForEachVisibleOutputPin(Node, [&Result](UEdGraphPin* Pin)
+		{
+			Result.Add(UBlueprintGraphPin::FromNativePin(Pin));
+			return true;
+		});
+	}
+	return Result;
+}
+
+UBlueprintGraphPin* UBlueprintEditorLibrary::FindInputPinByIndex(const UK2Node* Node, int32 Index, const FEdGraphPinType& Type)
+{
+	if(!::IsValid(Node))
+	{
+		return nullptr;
+	}
+
+	const UEdGraphSchema_K2* K2Schema = GetDefault<UEdGraphSchema_K2>();
+	UBlueprintGraphPin* Result = nullptr;
+	int32 Iter = 0; 
+	const UClass* Class = Node->GetBlueprintClassFromNode();
+	UE::Private::ForEachVisibleInputPin(Node, [&Iter, Class, Node, Index, &Type, K2Schema, &Result](UEdGraphPin* Pin) {
+		if(Iter != Index)
+		{
+			++Iter;
+			return true; // wrong index, keep searching
+		}
+
+		Result = UBlueprintGraphPin::FromNativePin(UE::Private::GetIfTypeMatches(Class, Pin, Type));
+		if(!Result)
+		{
+			UE_LOGFMT(LogBlueprint, Warning, 
+				"Type on input pin {pin_index} named {pin_name} on {node} did not match {request_type}",
+				Index, 
+				Pin->PinName, 
+				Node->GetPathName(), 
+				PinTypeToJsonSchema(Type, Class) );
+		}
+		return false; // index matched, all done
+	});
+	
+	if(!Result)
+	{
+		UE_LOGFMT(LogBlueprint, Warning, "Missing input pin at: {pin_index}, Pin count: {iter}, on node {node}", Index, Iter, Node->GetPathName());
+	}
+	return Result;
+}
+
+UBlueprintGraphPin* UBlueprintEditorLibrary::FindInputPin(const UK2Node* Node, FName PinName, const FEdGraphPinType& Type)
+{
+	if(!::IsValid(Node))
+	{
+		return nullptr;
+	}
+
+	UBlueprintGraphPin* Result = nullptr;
+	const UClass* Class = Node->GetBlueprintClassFromNode();
+	UE::Private::ForEachVisibleInputPin(Node, [PinName, Class, Node, &Type, &Result](UEdGraphPin* Pin) 
+	{
+		if(Pin->PinName != PinName)
+		{
+			return true;
+		}
+		
+		Result = UBlueprintGraphPin::FromNativePin(UE::Private::GetIfTypeMatches(Class, Pin, Type));
+		if(!Result)
+		{
+			UE_LOGFMT(LogBlueprint, Warning, 
+				"Type of input pin {pin_name} on {node} did not match {request_type}",
+				Pin->PinName, 
+				Node->GetPathName(), 
+				PinTypeToJsonSchema(Type, Class) );
+		}
+		return false;
+	});
+	
+	if(!Result)
+	{
+		UE_LOGFMT(LogBlueprint, Warning, "Missing input pin named: {pin_name}, on node {node}", PinName, Node->GetPathName());
+	}
+	return Result;
+}
+
+UBlueprintGraphPin* UBlueprintEditorLibrary::FindOutputPinByIndex(const UK2Node* Node, int32 Index, const FEdGraphPinType& Type)
+{
+	if(!::IsValid(Node))
+	{
+		return nullptr;
+	}
+
+	const UEdGraphSchema_K2* K2Schema = GetDefault<UEdGraphSchema_K2>();
+	UBlueprintGraphPin* Result = nullptr;
+	int32 Iter = 0; 
+	const UClass* Class = Node->GetBlueprintClassFromNode();
+	UE::Private::ForEachVisibleOutputPin(Node, [&Iter, Class, Node, Index, &Type, K2Schema, &Result](UEdGraphPin* Pin) 
+	{
+		if(Iter != Index)
+		{
+			++Iter;
+			return true; // wrong index, keep searching
+		}
+		
+		Result = UBlueprintGraphPin::FromNativePin(UE::Private::GetIfTypeMatches(Class, Pin, Type));
+		if(!Result)
+		{
+			UE_LOGFMT(LogBlueprint, Warning, 
+				"Type on output pin {pin_index} named {pin_name} on {node} did not match {request_type}",
+				Index, 
+				Pin->PinName, 
+				Node->GetPathName(), 
+				PinTypeToJsonSchema(Type, Class) );
+		}
+		return false; // index matched, all done
+	});
+	
+	if(!Result)
+	{
+		UE_LOGFMT(LogBlueprint, Warning, "Missing output pin at: {pin_index}, Pin count: {iter}, on node {node}", Index, Iter, Node->GetPathName());
+	}
+	return Result;
+}
+
+UBlueprintGraphPin* UBlueprintEditorLibrary::FindOutputPin(const UK2Node* Node, FName PinName, const FEdGraphPinType& Type)
+{
+	if(!::IsValid(Node))
+	{
+		return nullptr;
+	}
+
+	UBlueprintGraphPin* Result = nullptr;
+	const UClass* Class = Node->GetBlueprintClassFromNode();
+	UE::Private::ForEachVisibleOutputPin(Node, [Class, Node, PinName, &Type, &Result](UEdGraphPin* Pin) 
+	{
+		if(Pin->PinName != PinName)
+		{
+			return true;
+		}
+		
+		Result = UBlueprintGraphPin::FromNativePin(UE::Private::GetIfTypeMatches(Class, Pin, Type));
+		if(!Result)
+		{
+			UE_LOGFMT(LogBlueprint, Warning, 
+				"Type of output pin {pin_name} on {node} did not match {request_type}",
+				Pin->PinName, 
+				Node->GetPathName(), 
+				PinTypeToJsonSchema(Type, Class) );
+		}
+		return false;
+	});
+	
+	if(!Result)
+	{
+		UE_LOGFMT(LogBlueprint, Warning, "Missing output pin named: {pin_name}, on node {node}", PinName, Node->GetPathName());
+	}
+	return Result;
+}
+
+UBlueprintGraphPin* UBlueprintEditorLibrary::FindExecutePin(const UK2Node* Node)
+{
+	if(!::IsValid(Node))
+	{
+		return nullptr;
+	}
+
+	UBlueprintGraphPin* Result = nullptr;
+	UE::Private::ForEachVisibleInputPin(Node, [&Result](UEdGraphPin* Pin) 
+	{
+		if(UEdGraphSchema_K2::IsExecPin(*Pin) && Pin->PinName == UEdGraphSchema_K2::PN_Execute)
+		{
+			Result = UBlueprintGraphPin::FromNativePin(Pin);
+			return false;
+		}
+		return true;
+	});
+	if(!Result)
+	{
+		UE_LOGFMT(LogBlueprint, Warning, "No execute pin found on node {node}", Node->GetPathName());
+	}
+	return Result;
+}
+
+UBlueprintGraphPin* UBlueprintEditorLibrary::FindThenPin(const UK2Node* Node)
+{
+	if(!::IsValid(Node))
+	{
+		return nullptr;
+	}
+
+	UBlueprintGraphPin* Result = nullptr;
+	UE::Private::ForEachVisibleOutputPin(Node, [&Result](UEdGraphPin* Pin) 
+	{
+		if(UEdGraphSchema_K2::IsExecPin(*Pin) && Pin->PinName == UEdGraphSchema_K2::PN_Then)
+		{
+			Result = UBlueprintGraphPin::FromNativePin(Pin);
+			return false;
+		}
+		return true;
+	});
+	if(!Result)
+	{
+		UE_LOGFMT(LogBlueprint, Warning, "No then pin found on node {node}", Node->GetPathName());
+	}
+	return Result;
+}
+
+UBlueprintGraphPin* UBlueprintEditorLibrary::FindSelfPin(const UK2Node* Node)
+{
+	if(!::IsValid(Node))
+	{
+		return nullptr;
+	}
+
+	UBlueprintGraphPin* Result = nullptr;
+	const UEdGraphSchema_K2* K2Schema = GetDefault<UEdGraphSchema_K2>();
+	UE::Private::ForEachVisibleInputPin(Node, [&Result, K2Schema](UEdGraphPin* Pin) 
+	{
+		if(K2Schema->IsSelfPin(*Pin))
+		{
+			Result = UBlueprintGraphPin::FromNativePin(Pin);
+			return false;
+		}
+		return true;
+	});
+	if(!Result)
+	{
+		UE_LOGFMT(LogBlueprint, Warning, "No self pin found on node {node}", Node->GetPathName());
+	}
+	return Result;
+}
+
+UBlueprintGraphPin* UBlueprintEditorLibrary::FindResultPin(const UK2Node* Node)
+{
+	if(!::IsValid(Node))
+	{
+		return nullptr;
+	}
+
+	UEdGraphPin* Result = nullptr;
+	const UEdGraphSchema_K2* K2Schema = GetDefault<UEdGraphSchema_K2>();
+	bool bDidWarn = false;
+	UE::Private::ForEachVisibleOutputPin(Node, [&Result, Node, &bDidWarn, K2Schema](UEdGraphPin* Pin) 
+	{
+		if(UEdGraphSchema_K2::IsExecPin(*Pin))
+		{
+			return true;
+		}
+		if(Result)
+		{
+			Result = nullptr; // multiple outputs found, return nullptr
+			UE_LOGFMT(LogBlueprint, Warning, "Multiple result pins found on node {node}", Node->GetPathName());
+			bDidWarn = true;
+			return false;
+		}
+		Result = Pin;
+		return true;
+	});
+	if(!Result && !bDidWarn)
+	{
+		UE_LOGFMT(LogBlueprint, Warning, "No result pin found on node {node}", Node->GetPathName());
+	}
+	return UBlueprintGraphPin::FromNativePin(Result);
+}
+
+UBlueprintGraphPin* UBlueprintEditorLibrary::FindDataInputPin(const UK2Node* Node)
+{
+	if(!::IsValid(Node))
+	{
+		return nullptr;
+	}
+
+	UEdGraphPin* Result = nullptr;
+	const UEdGraphSchema_K2* K2Schema = GetDefault<UEdGraphSchema_K2>();
+	bool bDidWarn = false;
+	UE::Private::ForEachVisibleInputPin(Node, [&Result, &bDidWarn, Node, K2Schema](UEdGraphPin* Pin) 
+	{
+		if(UEdGraphSchema_K2::IsExecPin(*Pin))
+		{
+			return true;
+		}
+		if(Result)
+		{
+			Result = nullptr; // multiple inputs found, return nullptr
+			UE_LOGFMT(LogBlueprint, Warning, "Multiple input data pins found on node {node}", Node->GetPathName());
+			bDidWarn = true;
+			return false;
+		}
+		Result = Pin;
+		return true;
+	});
+	if(!Result && !bDidWarn)
+	{
+		UE_LOGFMT(LogBlueprint, Warning, "No data input pin found on node {node}", Node->GetPathName());
+	}
+	return UBlueprintGraphPin::FromNativePin(Result);
+}
+
+void UBlueprintEditorLibrary::SetNodePos(UK2Node* Node, FIntPoint Pos)
+{
+	if(!::IsValid(Node))
+	{
+		return;
+	}
+
+	Node->Modify();
+
+	Node->NodePosX = Pos.X;
+	Node->NodePosY = Pos.Y;
+}
+
+FIntPoint UBlueprintEditorLibrary::GetNodePos(const UK2Node* Node)
+{
+	if(!::IsValid(Node))
+	{
+		return {0, 0};
+	}
+
+	return {Node->NodePosX, Node->NodePosY};
+}
+
+UBlueprintGraphPin* UBlueprintEditorLibrary::FindElsePin(const UK2Node_IfThenElse* Node)
+{
+	if(!::IsValid(Node))
+	{
+		return nullptr;
+	}
+
+	return UBlueprintGraphPin::FromNativePin(Node->GetElsePin());
+}
+	
+UBlueprintGraphPin* UBlueprintEditorLibrary::FindConditionPin(const UK2Node_IfThenElse* Node)
+{
+	if(!::IsValid(Node))
+	{
+		return nullptr;
+	}
+
+	return UBlueprintGraphPin::FromNativePin(Node->GetConditionPin());
+}
+
 #undef LOCTEXT_NAMESPACE	// "BlueprintEditorLibrary"
diff --git a/Engine/Source/Editor/BlueprintEditorLibrary/Private/BlueprintEditorLibrary/BlueprintGraphEditor.cpp b/Engine/Source/Editor/BlueprintEditorLibrary/Private/BlueprintEditorLibrary/BlueprintGraphEditor.cpp
new file mode 100644
--- /dev/null
+++ b/Engine/Source/Editor/BlueprintEditorLibrary/Private/BlueprintEditorLibrary/BlueprintGraphEditor.cpp
@@ -0,0 +1,963 @@
+// Copyright Epic Games, Inc. All Rights Reserved.
+
+#include "BlueprintGraphEditor.h"
+
+#include "Algo/Find.h"
+#include "Algo/IndexOf.h"
+#include "BlueprintEditorLibrary.h"
+#include "BlueprintFunctionNodeSpawner.h"
+#include "BlueprintVariableNodeSpawner.h"
+#include "Blueprint/BlueprintExceptionInfo.h"
+#include "Kismet/KismetSystemLibrary.h"
+#include "K2Node_CallFunction.h"
+#include "K2Node_FunctionEntry.h"
+#include "K2Node_FunctionResult.h"
+#include "K2Node_MacroInstance.h"
+#include "K2Node_IfThenElse.h"
+#include "K2Node_Select.h"
+#include "K2Node_FunctionResult.h"
+#include "K2Node_VariableGet.h"
+#include "K2Node_VariableSet.h"
+#include "Kismet2/BlueprintEditorUtils.h"
+#include "Logging/StructuredLog.h"
+#include "UObject/Class.h"
+#include "UObject/SoftObjectPtr.h"
+#include "UObject/SoftObjectPath.h"
+#include "BlueprintEditorLibrary/BlueprintGraphPin.h"
+
+#define LOCTEXT_NAMESPACE "UBlueprintGraphEditor"
+
+UBlueprintGraphEditor::UBlueprintGraphEditor()
+{
+}
+
+UBlueprintGraphEditor* UBlueprintGraphEditor::CreateAndEditFunctionGraph(UBlueprint* InBlueprint, const FString& FuncName)
+{
+	if(!InBlueprint)
+	{
+		return nullptr;
+	}
+
+	UEdGraph* InGraph = UBlueprintEditorLibrary::AddFunctionGraph(InBlueprint, FuncName);
+	if(!InGraph)
+	{
+		return nullptr;
+	}
+
+	UBlueprintGraphEditor* GraphEditor = NewObject<UBlueprintGraphEditor>();
+	GraphEditor->Blueprint = InBlueprint;
+	GraphEditor->Graph = InGraph;
+	if(InGraph->Nodes.Num() > 0 )
+	{
+		UEdGraphNode* Node = InGraph->Nodes[0];
+		if(Node)
+		{
+			GraphEditor->SpawnLocation = FVector2D(Node->NodePosX, Node->NodePosY);
+		}
+	}
+	return GraphEditor;
+}
+
+UBlueprintGraphEditor* UBlueprintGraphEditor::GetGraphEditor(UBlueprint* InBlueprint, FName GraphName)
+{
+	const auto FindGraphByName = [](const TArray<UEdGraph*>& GraphList, FName Name)
+	{
+		return Algo::FindByPredicate(GraphList, [Name](UEdGraph* Graph)
+		{
+			return Graph && Graph->GetFName() == Name;
+		} );
+	};
+	if(!InBlueprint)
+	{
+		return nullptr;
+	}
+
+	TArray<UEdGraph*> AllGraphs;
+	InBlueprint->GetAllGraphs(AllGraphs);
+	UEdGraph *const* FoundGraph = FindGraphByName(AllGraphs, GraphName);
+	if(!FoundGraph)
+	{
+		UE_LOGFMT(LogBlueprint, Error, "Failed to find graph named {GraphName} for editing in {Blueprint}", GraphName, InBlueprint->GetPathName());
+		return nullptr;
+	}
+
+	UBlueprintGraphEditor* GraphEditor = NewObject<UBlueprintGraphEditor>();
+	GraphEditor->Blueprint = InBlueprint;
+	GraphEditor->Graph = *FoundGraph;
+	return GraphEditor;
+}
+
+void UBlueprintGraphEditor::ListAllNodes(TArray<UK2Node*>& OutNodes)
+{
+	if(!IsValid())
+	{
+		return;
+	}
+
+	Graph->GetNodesOfClass(OutNodes);
+}
+	
+void UBlueprintGraphEditor::ListNodesOfClass(TArray<UK2Node*>& OutNodes, TSubclassOf<UK2Node> Class)
+{
+	if(!IsValid() || !Class)
+	{
+		return;
+	}
+
+	for(UEdGraphNode* Node : Graph->Nodes)
+	{
+		if(Node->IsA(Class))
+		{
+			OutNodes.Add(CastChecked<UK2Node>(Node));
+		}
+	}
+}
+
+void UBlueprintGraphEditor::ListNodesWithErrors(TArray<UK2Node*>& OutNodes)
+{
+	if(!IsValid())
+	{
+		return;
+	}
+
+	for(UEdGraphNode* Node : Graph->Nodes)
+	{
+		if(UK2Node* K2Node = Cast<UK2Node>(Node))
+		{
+			if (K2Node->bHasCompilerMessage && !Node->ErrorMsg.IsEmpty() && Node->ErrorType == EMessageSeverity::Error)
+			{
+				OutNodes.Add(K2Node);
+			}
+		}
+	}
+}
+
+void UBlueprintGraphEditor::ListNodesWithWarnings(TArray<UK2Node*>& OutNodes)
+{
+	if(!IsValid())
+	{
+		return;
+	}
+
+	for(UEdGraphNode* Node : Graph->Nodes)
+	{
+		if(UK2Node* K2Node = Cast<UK2Node>(Node))
+		{
+			if (K2Node->bHasCompilerMessage && !Node->ErrorMsg.IsEmpty() && Node->ErrorType == EMessageSeverity::Warning)
+			{
+				OutNodes.Add(K2Node);
+			}
+		}
+	}
+}
+
+void UBlueprintGraphEditor::ListNodesWithNotes(TArray<UK2Node*>& OutNodes)
+{
+	if(!IsValid())
+	{
+		return;
+	}
+	
+	for(UEdGraphNode* Node : Graph->Nodes)
+	{
+		if(UK2Node* K2Node = Cast<UK2Node>(Node))
+		{
+			if (K2Node->bHasCompilerMessage && !Node->ErrorMsg.IsEmpty() && Node->ErrorType == EMessageSeverity::Info)
+			{
+				OutNodes.Add(K2Node);
+			}
+		}
+	}
+}
+
+bool UBlueprintGraphEditor::RemoveGraphInputParameter(FName InputName)
+{
+	if(!IsValid())
+	{
+		return false;
+	}
+
+	UK2Node_FunctionEntry* Entry = GetFunctionEntryNode();
+	if(!Entry)
+	{
+		return false;
+	}
+	Entry->Modify();
+	const bool bHasPin = Entry->FindPin(InputName, EGPD_Output) != nullptr;
+	if(bHasPin)
+	{
+		Entry->RemoveUserDefinedPinByName(InputName);
+	}
+	return bHasPin;
+}
+
+UBlueprintGraphPin* UBlueprintGraphEditor::FindGraphEntryPin() const
+{
+	UK2Node_FunctionEntry* Entry = GetFunctionEntryNode();
+	if(!Entry)
+	{
+		return nullptr;
+	}
+	return UBlueprintEditorLibrary::FindThenPin(Entry);//Entry->FindThenPin();
+}
+
+UK2Node_FunctionResult* UBlueprintGraphEditor::AddGraphOutputParameter(FName OutputName, const FEdGraphPinType& PinType)
+{
+	if(!IsValid())
+	{
+		return nullptr;
+	}
+
+	const UEdGraphSchema_K2* Schema = GetDefault<UEdGraphSchema_K2>();
+	TArray<UK2Node_FunctionResult*> ResultNodes;
+	Graph->GetNodesOfClass(ResultNodes);
+	if(ResultNodes.Num() == 0)
+	{
+		if(UK2Node_FunctionResult* Result = AddReturnNode())
+		{
+			ResultNodes.Add(Result);
+		}
+		else
+		{
+			UE_LOGFMT(LogBlueprint, Error, "Failed to create a result node in order to create an output parameter");
+			return nullptr;
+		}
+	}
+
+	for(UK2Node_FunctionResult* Result : ResultNodes)
+	{
+		Result->Modify();
+		UEdGraphPin* Pin = Result->CreateUserDefinedPin(OutputName, 
+			PinType != FEdGraphPinType() ? PinType : UBlueprintEditorLibrary::GetBasicTypeByName(UEdGraphSchema_K2::PC_Boolean), 
+			EGPD_Input);
+	}
+	FBlueprintEditorUtils::MarkBlueprintAsStructurallyModified(Blueprint);
+
+	return ResultNodes[0];
+}
+
+bool UBlueprintGraphEditor::RemoveGraphOutputParameter(FName OutputName)
+{
+	if(!IsValid())
+	{
+		return false;
+	}
+
+	TArray<UK2Node_FunctionResult*> TargetNodes;
+	Graph->GetNodesOfClass(TargetNodes);
+
+	bool bDoesAnyResultNodeHavePin = false;
+	for (UK2Node_EditablePinBase* TargetNode : TargetNodes)
+	{
+		if(TargetNode->UserDefinedPinExists(OutputName))
+		{
+			bDoesAnyResultNodeHavePin = true;
+			break;
+		}
+	}
+	if(!bDoesAnyResultNodeHavePin)
+	{
+		return false;
+	}
+
+	for (UK2Node_EditablePinBase* TargetNode : TargetNodes)
+	{
+		TargetNode->Modify();
+	}
+	
+	for (UK2Node_EditablePinBase* TargetNode : TargetNodes)
+	{
+		TargetNode->RemoveUserDefinedPinByName(OutputName);
+	}
+	
+	const UEdGraphSchema_K2* K2Schema = GetDefault<UEdGraphSchema_K2>();
+	
+	for (UK2Node_EditablePinBase* TargetNode : TargetNodes)
+	{
+		const bool bCurDisableOrphanSaving = TargetNode->bDisableOrphanPinSaving;
+		TargetNode->bDisableOrphanPinSaving = true;
+		TargetNode->ReconstructNode();
+		TargetNode->bDisableOrphanPinSaving = bCurDisableOrphanSaving;
+
+		K2Schema->HandleParameterDefaultValueChanged(TargetNode);
+	}
+	return true;
+}
+
+bool UBlueprintGraphEditor::AddLocalVariable(FName LocalName, const FEdGraphPinType& PinType, const FString& Value)
+{
+	if(!IsValid())
+	{
+		return false;
+	}
+
+	FName VarName = FBlueprintEditorUtils::FindUniqueKismetName(
+		Blueprint, LocalName.ToString(), FindUField<UFunction>(Blueprint->SkeletonGeneratedClass, Graph->GetFName()));
+	return FBlueprintEditorUtils::AddLocalVariable(
+		Blueprint, 
+		Graph, 
+		VarName, 
+		PinType != FEdGraphPinType() ? 
+			PinType : UBlueprintEditorLibrary::GetBasicTypeByName(UEdGraphSchema_K2::PC_Boolean),
+		Value);
+}
+
+bool UBlueprintGraphEditor::RemoveLocalVariable(FName LocalName)
+{
+	if(!IsValid())
+	{
+		return false;
+	}
+
+	UStruct* Function = FindUField<UFunction>(Blueprint->SkeletonGeneratedClass, Graph->GetFName());
+	if(!Function)
+	{
+		return false;
+	}
+	if(FBlueprintEditorUtils::FindLocalVariable(Blueprint, 
+		Function, LocalName) == nullptr)
+	{
+		return false;
+	}
+		
+	FBlueprintEditorUtils::RemoveLocalVariable(Blueprint, Function, LocalName);
+	return true;
+}
+
+bool UBlueprintGraphEditor::AddMemberVariable(FName MemberName, const FEdGraphPinType& PinType, const FString& Value)
+{
+	if(!IsValid())
+	{
+		return false;
+	}
+
+	FName VarName = FBlueprintEditorUtils::FindUniqueKismetName(
+		Blueprint, MemberName.ToString(), FindUField<UFunction>(Blueprint->SkeletonGeneratedClass, Graph->GetFName()));
+	return FBlueprintEditorUtils::AddMemberVariable(
+		Blueprint, 
+		VarName,
+		PinType != FEdGraphPinType() ? 
+			PinType : UBlueprintEditorLibrary::GetBasicTypeByName(UEdGraphSchema_K2::PC_Boolean),
+		Value);
+}
+
+bool UBlueprintGraphEditor::RemoveMemberVariable(FName MemberName)
+{
+	if(!IsValid())
+	{
+		return false;
+	}
+
+	const int32 VarIndex = FBlueprintEditorUtils::FindNewVariableIndex(Blueprint, MemberName);
+	if (VarIndex == INDEX_NONE)
+	{
+		return false;
+	}
+
+	FBlueprintEditorUtils::RemoveMemberVariable(Blueprint, MemberName);
+	return true;
+}
+
+void UBlueprintGraphEditor::SetIsPureFunction(bool bIsPure)
+{
+	if(!IsValid())
+	{
+		return;
+	}
+
+	UFunction* Function = GetSkeletonFunction();
+	if(!Function)
+	{
+		return;
+	}
+
+	UK2Node_FunctionEntry* EntryNode = GetFunctionEntryNode();
+	SetFunctionFlagImpl(Function, EntryNode, FUNC_BlueprintPure, bIsPure);
+}
+
+void UBlueprintGraphEditor::SetIsCallInEditorFunction(bool bCallInEditor)
+{
+	if(!IsValid())
+	{
+		return;
+	}
+
+	UK2Node_FunctionEntry* EntryNode = GetFunctionEntryNode();
+	if (EntryNode)
+	{
+		EntryNode->Modify();
+		EntryNode->MetaData.bCallInEditor = bCallInEditor;
+	}
+}
+
+void UBlueprintGraphEditor::SetFunctionIsPublic()
+{
+	if(!IsValid())
+	{
+		return;
+	}
+
+	UFunction* Function = GetSkeletonFunction();
+	if(!Function)
+	{
+		return;
+	}
+
+	UK2Node_FunctionEntry* EntryNode = GetFunctionEntryNode();
+	SetFunctionFlagImpl(Function, EntryNode, FUNC_Private, false);
+	SetFunctionFlagImpl(Function, EntryNode, FUNC_Protected, false);
+	SetFunctionFlagImpl(Function, EntryNode, FUNC_Public, true);
+}
+
+void UBlueprintGraphEditor::SetFunctionIsProtected()
+{
+	if(!IsValid())
+	{
+		return;
+	}
+
+	UFunction* Function = GetSkeletonFunction();
+	if(!Function)
+	{
+		return;
+	}
+
+	UK2Node_FunctionEntry* EntryNode = GetFunctionEntryNode();
+	SetFunctionFlagImpl(Function, EntryNode, FUNC_Private, false);
+	SetFunctionFlagImpl(Function, EntryNode, FUNC_Protected, true);
+	SetFunctionFlagImpl(Function, EntryNode, FUNC_Public, false);
+}
+
+void UBlueprintGraphEditor::SetFunctionIsPrivate()
+{
+	if(!IsValid())
+	{
+		return;
+	}
+
+	UFunction* Function = GetSkeletonFunction();
+	if(!Function)
+	{
+		return;
+	}
+
+	UK2Node_FunctionEntry* EntryNode = GetFunctionEntryNode();
+	SetFunctionFlagImpl(Function, EntryNode, FUNC_Private, true);
+	SetFunctionFlagImpl(Function, EntryNode, FUNC_Protected, false);
+	SetFunctionFlagImpl(Function, EntryNode, FUNC_Public, false);
+}
+
+void UBlueprintGraphEditor::SetIsConstFunction(bool bIsConst)
+{
+	if(!IsValid())
+	{
+		return;
+	}
+
+	UFunction* Function = GetSkeletonFunction();
+	if(!Function)
+	{
+		return;
+	}
+
+	UK2Node_FunctionEntry* EntryNode = GetFunctionEntryNode();
+	SetFunctionFlagImpl(Function, EntryNode, FUNC_Const, bIsConst);
+}
+
+void UBlueprintGraphEditor::SetIsExecFunction(bool bIsExec)
+{
+	if(!IsValid())
+	{
+		return;
+	}
+
+	UFunction* Function = GetSkeletonFunction();
+	if(!Function)
+	{
+		return;
+	}
+
+	UK2Node_FunctionEntry* EntryNode = GetFunctionEntryNode();
+	SetFunctionFlagImpl(Function, EntryNode, FUNC_Exec, bIsExec);
+}
+
+void UBlueprintGraphEditor::SetIsThreadSafeFunction(bool bIsThreadSafe)
+{
+	if(!IsValid())
+	{
+		return;
+	}
+
+	UK2Node_FunctionEntry* EntryNode = GetFunctionEntryNode();
+	if (EntryNode)
+	{
+		EntryNode->Modify();
+		EntryNode->MetaData.bThreadSafe = bIsThreadSafe;
+	}
+}
+
+void UBlueprintGraphEditor::SetIsUnsafeDuringActorConstructionFunction(bool bIsUnsafeDuringActorConstruction)
+{
+	if(!IsValid())
+	{
+		return;
+	}
+
+	UK2Node_FunctionEntry* EntryNode = GetFunctionEntryNode();
+	if (EntryNode)
+	{
+		EntryNode->Modify();
+		EntryNode->MetaData.bIsUnsafeDuringActorConstruction = bIsUnsafeDuringActorConstruction;
+	}
+}
+
+void UBlueprintGraphEditor::SetIsDeprecatedFunction(bool bIsDeprecated)
+{
+	if(!IsValid())
+	{
+		return;
+	}
+
+	UK2Node_FunctionEntry* EntryNode = GetFunctionEntryNode();
+	if (EntryNode)
+	{
+		EntryNode->Modify();
+		EntryNode->MetaData.bIsDeprecated = bIsDeprecated;
+	}
+}
+
+void UBlueprintGraphEditor::SetDeprecationMessageOnFunction(FText Message)
+{
+	if(!IsValid())
+	{
+		return;
+	}
+
+	UK2Node_FunctionEntry* EntryNode = GetFunctionEntryNode();
+	if (EntryNode)
+	{
+		EntryNode->Modify();
+		EntryNode->MetaData.DeprecationMessage = Message.ToString();
+	}
+}
+
+
+UK2Node_CallFunction* UBlueprintGraphEditor::AddCallFunctionNode(const FString& FunctionPath)
+{
+	if(!IsValid())
+	{
+		return nullptr;
+	}
+
+	TSoftObjectPtr<UFunction> FunctionSoftPtr = TSoftObjectPtr<UFunction>(FSoftObjectPath(FunctionPath));
+	const UFunction* Function = FunctionSoftPtr.LoadSynchronous();
+	if (!Function && Blueprint->SkeletonGeneratedClass)
+	{
+		// check for a local function:
+		Function = Blueprint->SkeletonGeneratedClass->FindFunctionByName(*FunctionPath);
+	}
+
+	if(!Function)
+	{
+		return nullptr;
+	}
+
+	UBlueprintFunctionNodeSpawner* Spawner = UBlueprintFunctionNodeSpawner::Create(Function);
+	UK2Node_CallFunction* Result = CastChecked<UK2Node_CallFunction>(Spawner->Invoke(Graph, {}, GetNodeLocation()));
+	if(!Result)
+	{
+		return nullptr;
+	}
+
+	HandleNodeSpawned(Result);
+	return Result;
+}
+
+UK2Node_MacroInstance* UBlueprintGraphEditor::AddMacroNode(const FString& MacroPath)
+{
+	if(!IsValid())
+	{
+		return nullptr;
+	}
+
+	TSoftObjectPtr<UEdGraph> GraphSoftPtr = TSoftObjectPtr<UEdGraph>(FSoftObjectPath(MacroPath));
+	UEdGraph* MacroGraph = GraphSoftPtr.LoadSynchronous();
+	if(!MacroGraph)
+	{
+		UE_LOGFMT(LogBlueprint, Error, "Failed to find macro graph named {MacroPath}", MacroPath);
+		return nullptr;
+	}
+
+	UBlueprintNodeSpawner* NodeSpawner = UBlueprintNodeSpawner::Create(UK2Node_MacroInstance::StaticClass());
+	check(NodeSpawner != nullptr);
+
+	NodeSpawner->CustomizeNodeDelegate = UBlueprintNodeSpawner::FCustomizeNodeDelegate::CreateLambda(
+		[MacroGraph](UEdGraphNode* NewNode, bool bIsTemplateNode)
+		{
+			UK2Node_MacroInstance* MacroNode = CastChecked<UK2Node_MacroInstance>(NewNode);
+			MacroNode->SetMacroGraph(MacroGraph);
+		}
+	);
+	UK2Node_MacroInstance* Result = CastChecked<UK2Node_MacroInstance>(NodeSpawner->Invoke(Graph, {}, GetNodeLocation()));
+	if(!Result)
+	{
+		return nullptr;
+	}
+
+	HandleNodeSpawned(Result);
+	return Result;
+}
+
+UK2Node_IfThenElse* UBlueprintGraphEditor::AddBranchNode()
+{
+	if(!IsValid())
+	{
+		return nullptr;
+	}
+
+	UBlueprintNodeSpawner* NodeSpawner = UBlueprintNodeSpawner::Create(UK2Node_IfThenElse::StaticClass());
+	UK2Node_IfThenElse* Result = CastChecked<UK2Node_IfThenElse>(NodeSpawner->Invoke(Graph, {}, GetNodeLocation()));
+	if(!Result)
+	{
+		return nullptr;
+	}
+
+	HandleNodeSpawned(Result);
+	return Result;
+}
+
+UK2Node_FunctionResult* UBlueprintGraphEditor::AddReturnNode()
+{
+	if(!IsValid())
+	{
+		return nullptr;
+	}
+
+	UBlueprintNodeSpawner* NodeSpawner = UBlueprintNodeSpawner::Create(UK2Node_FunctionResult::StaticClass());
+	UK2Node_FunctionResult* Result = CastChecked<UK2Node_FunctionResult>(NodeSpawner->Invoke(Graph, {}, GetNodeLocation()));
+	if(!Result)
+	{
+		return nullptr;
+	}
+
+	HandleNodeSpawned(Result);
+	return Result;
+}
+
+UK2Node_VariableGet* UBlueprintGraphEditor::AddGetMemberVariableNode(FName MemberName, const FString& ClassPath)
+{
+	if(!IsValid())
+	{
+		return nullptr;
+	}
+
+	const UClass* Scope = Blueprint->SkeletonGeneratedClass;
+	if(!ClassPath.IsEmpty())
+	{
+		Scope = GetClassFromPath(ClassPath);
+
+		if(!Scope)
+		{
+			UE_LOGFMT(LogBlueprint, Error, "Failed to find class {ClassPath}", ClassPath);
+			return nullptr;
+		}
+	}
+
+	if(!Scope)
+	{
+		UE_LOGFMT(LogBlueprint, Error, "Failed to resolve scope");
+		return nullptr;
+	}
+
+	const FProperty* Property = Scope->FindPropertyByName(MemberName);
+	if(!Property)
+	{
+		UE_LOGFMT(LogBlueprint, Error, "Failed to find member named {MemberName}", MemberName);
+		return nullptr;
+	}
+	
+	UBlueprintVariableNodeSpawner* Spawner = UBlueprintVariableNodeSpawner::CreateFromMemberOrParam(UK2Node_VariableGet::StaticClass(), Property, Graph);
+	UK2Node_VariableGet* Result = CastChecked<UK2Node_VariableGet>(Spawner->Invoke(Graph, {}, GetNodeLocation()));
+	if(!Result)
+	{
+		return nullptr;
+	}
+
+	HandleNodeSpawned(Result);
+	return Result;
+}
+
+UK2Node_VariableSet* UBlueprintGraphEditor::AddSetMemberVariableNode(FName MemberName, const FString& ClassPath)
+{
+	if(!IsValid())
+	{
+		return nullptr;
+	}
+
+	const UClass* Scope = Blueprint->SkeletonGeneratedClass;
+	if(!ClassPath.IsEmpty())
+	{
+		Scope = GetClassFromPath(ClassPath);
+	}
+
+	if(!::IsValid(Scope))
+	{
+		UE_LOGFMT(LogBlueprint, Error, "Failed to find Class for AddSetMemberVariableNode: {ClassPath}", ClassPath);
+		return nullptr;
+	}
+
+	const FProperty* Property = Scope->FindPropertyByName(MemberName);
+	if(!Property)
+	{
+		UE_LOGFMT(LogBlueprint, Error, "Failed to find member for AddSetMemberVariableNode: {MemberName}", MemberName);
+		return nullptr;
+	}
+	
+	UBlueprintVariableNodeSpawner* Spawner = UBlueprintVariableNodeSpawner::CreateFromMemberOrParam(UK2Node_VariableSet::StaticClass(), Property, Graph);
+	UK2Node_VariableSet* Result = CastChecked<UK2Node_VariableSet>(Spawner->Invoke(Graph, {}, GetNodeLocation()));
+	if(!Result)
+	{
+		return nullptr;
+	}
+
+	HandleNodeSpawned(Result);
+	return Result;
+}
+
+UK2Node_VariableGet* UBlueprintGraphEditor::AddGetLocalVariableNode(FName LocalName)
+{
+	if(!IsValid())
+	{
+		return nullptr;
+	}
+
+	UBlueprintVariableNodeSpawner* Spawner = nullptr;
+	FBPVariableDescription* VariableDescription = FBlueprintEditorUtils::FindLocalVariable(Blueprint, Graph, LocalName);
+	if(VariableDescription)
+	{
+		// create from local
+		FMemberReference Reference;
+		Reference.SetLocalMember(LocalName, Graph->GetName(), VariableDescription->VarGuid);
+		Spawner = UBlueprintVariableNodeSpawner::CreateFromLocal(
+			UK2Node_VariableGet::StaticClass(), Graph, *VariableDescription, Reference.ResolveMember<FProperty>(Blueprint->SkeletonGeneratedClass));	
+	}
+	else if(const FProperty* Input = FindParameterByName(LocalName))
+	{
+		// create from a parameter:
+		Spawner = UBlueprintVariableNodeSpawner::CreateFromMemberOrParam(UK2Node_VariableGet::StaticClass(), Input, Graph);
+	}
+
+	if(!Spawner)
+	{
+		return nullptr;
+	}
+	
+	UK2Node_VariableGet* Result = CastChecked<UK2Node_VariableGet>(Spawner->Invoke(Graph, {}, GetNodeLocation()));
+	if(!Result)
+	{
+		return nullptr;
+	}
+
+	HandleNodeSpawned(Result);
+	return Result;
+}
+
+UK2Node_VariableSet* UBlueprintGraphEditor::AddSetLocalVariableNode(FName LocalName)
+{
+	if(!IsValid())
+	{
+		return nullptr;
+	}
+
+	FBPVariableDescription* VariableDescription = FBlueprintEditorUtils::FindLocalVariable(Blueprint, Graph, LocalName);
+	if(!VariableDescription)
+	{
+		return nullptr;
+	}
+	
+	FMemberReference Reference;
+	Reference.SetLocalMember(LocalName, Graph->GetName(), VariableDescription->VarGuid);
+	UBlueprintVariableNodeSpawner* Spawner = UBlueprintVariableNodeSpawner::CreateFromLocal(
+		UK2Node_VariableSet::StaticClass(), Graph, *VariableDescription, Reference.ResolveMember<FProperty>(Blueprint->SkeletonGeneratedClass));
+	UK2Node_VariableSet* Result = CastChecked<UK2Node_VariableSet>(Spawner->Invoke(Graph, {}, GetNodeLocation()));
+	if(!Result)
+	{
+		return nullptr;
+	}
+
+	HandleNodeSpawned(Result);
+	return Result;
+}
+
+void UBlueprintGraphEditor::RemoveNodes(const TArray<UK2Node*>& Nodes)
+{
+	if(!IsValid())
+	{
+		return;
+	}
+
+	bool bNeedToModifyStructurally = false;
+	for(UK2Node* Node : Nodes)
+	{
+		if(!Node)
+		{
+			continue;
+		}
+
+		if (Node->CanUserDeleteNode())
+		{
+			if (Node->NodeCausesStructuralBlueprintChange())
+			{
+				bNeedToModifyStructurally = true;
+			}
+				
+			FBlueprintEditorUtils::RemoveNode(Blueprint, Node, true);
+		}
+	}
+
+	if (bNeedToModifyStructurally)
+	{
+		FBlueprintEditorUtils::MarkBlueprintAsStructurallyModified(Blueprint);
+	}
+	else
+	{
+		FBlueprintEditorUtils::MarkBlueprintAsModified(Blueprint);
+	}
+}
+
+UBlueprintGraphPin* UBlueprintGraphEditor::AddGraphInputParameter(FName InputName, const FEdGraphPinType& PinType, const FString& Value)
+{
+	if(!IsValid())
+	{
+		return nullptr;
+	}
+
+	UK2Node_FunctionEntry* Entry = GetFunctionEntryNode();
+	if(!Entry)
+	{
+		return nullptr;
+	}
+
+	const UEdGraphSchema_K2* Schema = GetDefault<UEdGraphSchema_K2>();
+	Entry->Modify();
+	UEdGraphPin* Pin = Entry->CreateUserDefinedPin(InputName, 
+		PinType != FEdGraphPinType() ? PinType : UBlueprintEditorLibrary::GetBasicTypeByName(UEdGraphSchema_K2::PC_Boolean), 
+		EGPD_Output);
+	if(Pin)
+	{
+		Schema->TrySetDefaultValue(*Pin, Value, false);
+		return UBlueprintGraphPin::FromNativePin(Pin);
+	}
+
+	return nullptr;
+}
+
+bool UBlueprintGraphEditor::IsValid() const
+{
+	return ::IsValid(Graph) && ::IsValid(Blueprint);
+}
+
+FVector2D UBlueprintGraphEditor::GetNodeLocation()
+{
+	const double GRID_SIZE = 120.0;
+	const double X_LIMIT = 2000.0;
+
+	SpawnLocation += FVector2D(GRID_SIZE, 0.0);
+	if(SpawnLocation.X > X_LIMIT)
+	{
+		SpawnLocation.X = 0.0;
+		SpawnLocation.Y += GRID_SIZE;
+	}
+	return SpawnLocation;
+}
+
+UK2Node_FunctionEntry* UBlueprintGraphEditor::GetFunctionEntryNode() const
+{
+	if(!IsValid())
+	{
+		return nullptr;
+	}
+
+	TArray<UK2Node_FunctionEntry*> EntryNodes;
+	Graph->GetNodesOfClass(EntryNodes);
+	if(EntryNodes.Num() != 1)
+	{
+		return nullptr;
+	}
+
+	return EntryNodes[0];
+}
+
+void UBlueprintGraphEditor::HandleNodeSpawned(UK2Node* Result)
+{
+	UE_LOGFMT(LogBlueprint, Display, "Node {node_path} spawn at ({pos_x},{pox_y})", Result->GetPathName(), Result->NodePosX, Result->NodePosY);
+}
+
+const FProperty* UBlueprintGraphEditor::FindParameterByName(FName Name) const
+{
+	const UFunction* SkeletonFunction = GetSkeletonFunction();
+	if(!SkeletonFunction)
+	{
+		return nullptr;
+	}
+
+	for (TFieldIterator<FProperty> ParamIt(SkeletonFunction); ParamIt && (ParamIt->PropertyFlags & CPF_Parm); ++ParamIt)
+	{
+		FProperty* Param = *ParamIt;
+		const bool bIsFunctionInput = !Param->HasAnyPropertyFlags(CPF_ReturnParm) && (!Param->HasAnyPropertyFlags(CPF_OutParm) || Param->HasAnyPropertyFlags(CPF_ReferenceParm));
+		if (bIsFunctionInput && Param->GetFName() == Name)
+		{
+			return *ParamIt;
+		}
+	}
+
+	return nullptr;
+}
+
+const UFunction* UBlueprintGraphEditor::GetSkeletonFunction() const
+{
+	return Blueprint->SkeletonGeneratedClass ? 
+		Blueprint->SkeletonGeneratedClass->FindFunctionByName(Graph->GetFName()) : 
+		nullptr;
+}
+
+UFunction* UBlueprintGraphEditor::GetSkeletonFunction()
+{
+	const UBlueprintGraphEditor* Self = this;
+	return const_cast<UFunction*>(Self->GetSkeletonFunction());
+}
+
+const UClass* UBlueprintGraphEditor::GetClassFromPath(const FString& ClassPath)
+{
+	TSoftObjectPtr<UClass> ClassSoftPtr = TSoftObjectPtr<UClass>(FSoftObjectPath(ClassPath));
+	const UClass* Class = ClassSoftPtr.LoadSynchronous();
+	if(Class)
+	{
+		return Class;
+	}
+	return nullptr;
+}
+
+void UBlueprintGraphEditor::SetFunctionFlagImpl(UFunction* Function, UK2Node_FunctionEntry* EntryNode, EFunctionFlags Flag, bool bIsSet)
+{
+	if (!ensure(EntryNode) || !ensure(Function))
+	{
+		return;
+	}
+
+	EntryNode->Modify();
+	Function->Modify();
+
+	if(bIsSet)
+	{
+		Function->FunctionFlags |= Flag;
+		EntryNode->SetExtraFlags(EntryNode->GetExtraFlags() | Flag);
+	}
+	else
+	{
+		Function->FunctionFlags &= ~Flag;
+		EntryNode->SetExtraFlags(EntryNode->GetExtraFlags() & ~Flag);
+	}
+}
+
+#undef LOCTEXT_NAMESPACE
diff --git a/Engine/Source/Editor/BlueprintEditorLibrary/Private/BlueprintEditorLibrary/BlueprintGraphEditor.h b/Engine/Source/Editor/BlueprintEditorLibrary/Private/BlueprintEditorLibrary/BlueprintGraphEditor.h
new file mode 100644
--- /dev/null
+++ b/Engine/Source/Editor/BlueprintEditorLibrary/Private/BlueprintEditorLibrary/BlueprintGraphEditor.h
@@ -0,0 +1,278 @@
+// Copyright Epic Games, Inc. All Rights Reserved.
+
+#pragma once
+
+#include "UObject/Object.h"
+#include "EdGraph/EdGraphNode.h"
+#include "EdGraph/EdGraphPin.h"
+#include "Templates/SubclassOf.h"
+
+#include "BlueprintGraphEditor.generated.h"
+
+class UEdGraph;
+class UBlueprint;
+
+class UK2Node;
+class UK2Node_CallFunction;
+class UK2Node_FunctionEntry;
+class UK2Node_MacroInstance;
+class UK2Node_IfThenElse;
+class UK2Node_Select;
+class UK2Node_FunctionResult;
+class UK2Node_VariableGet;
+class UK2Node_VariableSet;
+class UBlueprintGraphPin;
+
+/**
+ * Object oriented helper for editing graphs in a blueprint object
+ * Instantiate the UBlueprintGraphEditor via CreateAndEditFunctionGraph
+ * or GetGraphEditor
+ * 
+ * Python Example:
+
+import unreal
+
+BlueprintGraphEditor = unreal.BlueprintGraphEditor
+BlueprintEditorLibrary = unreal.BlueprintEditorLibrary
+BlueprintGraphPin = unreal.BlueprintGraphPin
+
+print_string = "/Script/Engine.KismetSystemLibrary.PrintString"
+
+def GenHelloWorld(bp):
+    editor = BlueprintGraphEditor.create_and_edit_function_graph(bp, "HelloWorld")
+    start_pin = editor.find_graph_entry_pin()
+    start_pos = start_pin.get_owning_node().get_node_pos()
+    print_node = editor.add_call_function_node(print_string)
+    print_node.find_input_pin("InString").set_pin_value("Generated Script")
+    editor.find_graph_entry_pin().try_create_connection(print_node.find_execute_pin())
+    print_node.set_node_pos(unreal.IntPoint(start_pos.x + 160, start_pos.y))
+
+ */
+UCLASS(Transient, BlueprintType, HideDropDown)
+class UBlueprintGraphEditor : public UObject
+{
+	GENERATED_BODY()
+private:
+	// This won't dissuade external NewObject calls, but HideDropDown keeps the editor clients away
+	// Please instantiate this object using one of the provided factory functions (e.g. CreateAndEditFunctionGraph)
+	UBlueprintGraphEditor(); 
+
+public:
+	/** returns a new UBlueprintGraphEditor for a graph in InBlueprint, creating a new graph named FuncName - useful for creating new function graphs */
+	UFUNCTION(BlueprintCallable, Category = "Blueprint Editor")
+	static UBlueprintGraphEditor* CreateAndEditFunctionGraph(UBlueprint* Blueprint, const FString& FuncName = FString(TEXT("NewFunction")));
+	
+	/** returns a new UBlueprintGraphEditor for a graph in InBlueprint named GraphName - useful for editing graphs that already exist */
+	UFUNCTION(BlueprintCallable, Category = "Blueprint Editor")
+	static UBlueprintGraphEditor* GetGraphEditor(UBlueprint* InBlueprint, FName GraphName);
+	
+	/**
+	 * Lists all nodes in the graph
+	 */
+	UFUNCTION(BlueprintCallable, Category = "Blueprint Editor")
+	void ListAllNodes(TArray<UK2Node*>& OutNodes);
+	/**
+	 * Lists nodes of the type provided by the Class input
+	 */
+	UFUNCTION(BlueprintCallable, Category = "Blueprint Editor")
+	void ListNodesOfClass(TArray<UK2Node*>& OutNodes, TSubclassOf<UK2Node> Class);
+	/**
+	 * Lists nodes with EMessageSeverity::Error level compiler messages
+	 */
+	UFUNCTION(BlueprintCallable, Category = "Blueprint Editor")
+	void ListNodesWithErrors(TArray<UK2Node*>& OutNodes);
+	/**
+	 * Lists nodes with EMessageSeverity::Warning level compiler messages
+	 */
+	UFUNCTION(BlueprintCallable, Category = "Blueprint Editor")
+	void ListNodesWithWarnings(TArray<UK2Node*>& OutNodes);
+	/**
+	 * Lists nodes with EMessageSeverity::Info level compiler messages
+	 */
+	UFUNCTION(BlueprintCallable, Category = "Blueprint Editor")
+	void ListNodesWithNotes(TArray<UK2Node*>& OutNodes);
+	
+	/**
+	 * Adds a graph input parameter with name InputName, and optional PinType and Value. Returns the added pin.
+	 * By default the type will be a boolean as defined by UBlueprintEditorLibrary.
+	 */
+	UFUNCTION(BlueprintCallable, Category = "Blueprint Editor")
+	UBlueprintGraphPin* AddGraphInputParameter(FName InputName, const FEdGraphPinType& PinType = FEdGraphPinType(), const FString& Value = FString());
+	
+	/**
+	 * Removes the graph input parameter specified by InputName
+	 */
+	UFUNCTION(BlueprintCallable, Category = "Blueprint Editor")
+	bool RemoveGraphInputParameter(FName InputName);
+	
+	/**
+	 * Returns the 'entry' pin for the graph, e.g. the 'then' pin of the entry node
+	 */
+	UFUNCTION(BlueprintCallable, Category = "Blueprint Editor")
+	UBlueprintGraphPin* FindGraphEntryPin() const;
+	
+	/**
+	 * Adds a graph output parameter, returning the UK2Node_FunctionResult that has the new output on it.
+	 * By default the type will be a boolean as defined by UBlueprintEditorLibrary.
+	 */
+	UFUNCTION(BlueprintCallable, Category = "Blueprint Editor")
+	UK2Node_FunctionResult* AddGraphOutputParameter(FName OutputName, const FEdGraphPinType& PinType = FEdGraphPinType());
+	
+	/**
+	 * Removes the output parameter specified by OutputName, returns true if there is a variable and it is removed
+	 */
+	UFUNCTION(BlueprintCallable, Category = "Blueprint Editor")
+	bool RemoveGraphOutputParameter(FName OutputName);
+	
+	/**
+	 * Adds a local variable, use BlueprintEditorLibrary to specify a type. A default value can optionally be provided.
+	 * By default the type will be a boolean as defined by UBlueprintEditorLibrary.
+	 */
+	UFUNCTION(BlueprintCallable, Category = "Blueprint Editor")
+	bool AddLocalVariable(FName LocalName, const FEdGraphPinType& PinType = FEdGraphPinType(), const FString& Value = FString());
+	/**
+	 * Removes the local variable with the corresponding name, returns true if there is a variable and it is removed
+	 */
+	UFUNCTION(BlueprintCallable, Category = "Blueprint Editor")
+	bool RemoveLocalVariable(FName LocalName);
+	
+	/**
+	 * Adds a member variable, use BlueprintEditorLibrary to specify a type. A default value can optionally be provided.
+	 * By default the type will be a boolean as defined by UBlueprintEditorLibrary.
+	 */
+	UFUNCTION(BlueprintCallable, Category = "Blueprint Editor")
+	bool AddMemberVariable(FName MemberName,  const FEdGraphPinType& PinType = FEdGraphPinType(), const FString& Value = FString());
+	/**
+	 * Removes the member variable with the corresponding name, returns true if there is a variable and it is removed
+	 */
+	UFUNCTION(BlueprintCallable, Category = "Blueprint Editor")
+	bool RemoveMemberVariable(FName MemberName);
+
+	/**
+	 * Sets the graph's FUNC_BlueprintPure flag
+	 */
+	UFUNCTION(BlueprintCallable, Category = "Blueprint Editor")
+	void SetIsPureFunction(bool bIsPure = true);
+	/**
+	 * Sets the graph's callability in editor
+	 */
+	UFUNCTION(BlueprintCallable, Category = "Blueprint Editor")
+	void SetIsCallInEditorFunction(bool bCallInEditor = true);
+	/**
+	 * Sets the graph's visibility to public
+	 */
+	UFUNCTION(BlueprintCallable, Category = "Blueprint Editor")
+	void SetFunctionIsPublic();
+	/**
+	 * Sets the graph's visibility to protected
+	 */
+	UFUNCTION(BlueprintCallable, Category = "Blueprint Editor")
+	void SetFunctionIsProtected();
+	/**
+	 * Sets the graph's visibility to private
+	 */
+	UFUNCTION(BlueprintCallable, Category = "Blueprint Editor")
+	void SetFunctionIsPrivate();
+	/**
+	 * Sets the graph's FUNC_Const flag
+	 */
+	UFUNCTION(BlueprintCallable, Category = "Blueprint Editor")
+	void SetIsConstFunction(bool bIsConst = true);
+	/**
+	 * Sets the graph's FUNC_Exec flag
+	 */
+	UFUNCTION(BlueprintCallable, Category = "Blueprint Editor")
+	void SetIsExecFunction(bool bIsExec = true);
+	/**
+	 * Sets the graph's thread safety status
+	 */
+	UFUNCTION(BlueprintCallable, Category = "Blueprint Editor")
+	void SetIsThreadSafeFunction(bool bIsThreadSafe = true);
+	/**
+	 * Sets the graph's construction script safety status
+	 */
+	UFUNCTION(BlueprintCallable, Category = "Blueprint Editor")
+	void SetIsUnsafeDuringActorConstructionFunction(bool bIsUnsafeDuringActorConstruction = true);
+	/**
+	 * Sets the graph's deprecation status
+	 */
+	UFUNCTION(BlueprintCallable, Category = "Blueprint Editor")
+	void SetIsDeprecatedFunction(bool bIsDeprecated = true);
+	/**
+	 * Sets the graph's deprecation message
+	 */
+	UFUNCTION(BlueprintCallable, Category = "Blueprint Editor")
+	void SetDeprecationMessageOnFunction(FText Message);
+
+	/**
+	 * Creates a UK2Node_CallFunction in the current graph invoking the function specified by FunctionPath
+	 */
+	UFUNCTION(BlueprintCallable, Category = "Blueprint Editor")
+	UK2Node_CallFunction* AddCallFunctionNode(const FString& FunctionPath);
+	/**
+	 * Creates a UK2Node_MacroInstance in the current graph for the provided MacroPath
+	 */
+	UFUNCTION(BlueprintCallable, Category = "Blueprint Editor")
+	UK2Node_MacroInstance* AddMacroNode(const FString& MacroPath);
+	
+	/**
+	 * Creates a UK2Node_IfThenElse in the current graph
+	 */
+	UFUNCTION(BlueprintCallable, Category = "Blueprint Editor")
+	UK2Node_IfThenElse* AddBranchNode();
+	/**
+	 * Creates a UK2Node_FunctionResult in the current graph
+	 */
+	UFUNCTION(BlueprintCallable, Category = "Blueprint Editor")
+	UK2Node_FunctionResult* AddReturnNode();
+
+	/**
+	 * Creates a UK2Node_VariableGet for the member variable identified by MemberName
+	 * by default this will use the current class, provided a ClassPath to work
+	 * with a different class
+	 */
+	UFUNCTION(BlueprintCallable, Category = "Blueprint Editor")
+	UK2Node_VariableGet* AddGetMemberVariableNode(FName MemberName, const FString& ClassPath = FString());
+	/**
+	 * Creates a UK2Node_VariableSet for the member variable identified by MemberName
+	 * by default this will use the current class, provided a ClassPath to work
+	 * with a different class
+	 */
+	UFUNCTION(BlueprintCallable, Category = "Blueprint Editor")
+	UK2Node_VariableSet* AddSetMemberVariableNode(FName MemberName, const FString& ClassPath = FString());
+	/**
+	 * Creates a UK2Node_VariableGet for the provided local variable identified by LocalName
+	 */
+	UFUNCTION(BlueprintCallable, Category = "Blueprint Editor")
+	UK2Node_VariableGet* AddGetLocalVariableNode(FName LocalName);
+	/**
+	 * Creates a UK2Node_VariableSet for the provided local variable identified by LocalName
+	 */
+	UFUNCTION(BlueprintCallable, Category = "Blueprint Editor")
+	UK2Node_VariableSet* AddSetLocalVariableNode(FName LocalName);
+
+	/**
+	 * Removes (aka deletes) the provided nodes from the current graph
+	 */
+	UFUNCTION(BlueprintCallable, Category = "Blueprint Editor")
+	void RemoveNodes(const TArray<UK2Node*>& Nodes);
+private:
+	bool IsValid() const;
+	FVector2D GetNodeLocation();
+	UK2Node_FunctionEntry* GetFunctionEntryNode() const;
+	void HandleNodeSpawned(UK2Node* Result);
+	const FProperty* FindParameterByName(FName Name) const;
+	const UFunction* GetSkeletonFunction() const;
+	UFunction* GetSkeletonFunction();
+	static const UClass* GetClassFromPath(const FString& ClassPath);
+	static void SetFunctionFlagImpl(UFunction* Function, UK2Node_FunctionEntry* EntryNode, EFunctionFlags Flag, bool bIsSet);
+
+	UPROPERTY()
+	FVector2D SpawnLocation;
+
+	UPROPERTY()
+	TObjectPtr<UBlueprint> Blueprint;
+
+	UPROPERTY()
+	TObjectPtr<UEdGraph> Graph;
+};
diff --git a/Engine/Source/Editor/BlueprintEditorLibrary/Private/BlueprintEditorLibrary/BlueprintGraphPin.cpp b/Engine/Source/Editor/BlueprintEditorLibrary/Private/BlueprintEditorLibrary/BlueprintGraphPin.cpp
new file mode 100644
--- /dev/null
+++ b/Engine/Source/Editor/BlueprintEditorLibrary/Private/BlueprintEditorLibrary/BlueprintGraphPin.cpp
@@ -0,0 +1,189 @@
+// Copyright Epic Games, Inc. All Rights Reserved.
+
+#include "BlueprintEditorLibrary/BlueprintGraphPin.h"
+
+#include "BlueprintEditorLibrary.h"
+#include "EdGraphSchema_K2.h"
+#include "K2Node.h"
+#include "Logging/StructuredLog.h"
+
+UBlueprintGraphPin* UBlueprintGraphPin::FromNativePin(UEdGraphPin* Pin)
+{
+	if(!Pin)
+	{
+		return nullptr;
+	}
+
+	UK2Node* OwningNode = Cast<UK2Node>(Pin->GetOwningNode());
+	if(!OwningNode)
+	{
+		return nullptr;
+	}
+
+	UBlueprintGraphPin* ScriptPin = NewObject<UBlueprintGraphPin>();
+	ScriptPin->Node = OwningNode;
+	ScriptPin->PinDirection = Pin->Direction;
+	ScriptPin->PinName = Pin->PinName;
+	return ScriptPin;
+}
+
+bool UBlueprintGraphPin::SetPinValue(const FString& Value)
+{						 
+	const UEdGraphSchema_K2* Schema = GetDefault<UEdGraphSchema_K2>();
+	if(UEdGraphPin* Pin = GetNativePin())
+	{
+		Schema->TrySetDefaultValue(*Pin, Value, false);
+		// Reset default values
+		if (!Schema->IsPinDefaultValid(Pin, Pin->DefaultValue, Pin->DefaultObject, Pin->DefaultTextValue).IsEmpty())
+		{
+			Schema->ResetPinToAutogeneratedDefaultValue(Pin, false);
+			UE_LOGFMT(LogBlueprint, Warning, "Failed to set default value {Value} on {PinName}", Value, Pin->PinName);
+			return false;
+		}
+		return true;
+	}
+	return false;
+}
+
+bool UBlueprintGraphPin::TryCreateConnection(UBlueprintGraphPin* Other)
+{
+	if(!Other)
+	{
+		return false;
+	}
+
+	UEdGraphPin* SelfNativePin = GetNativePin();
+	UEdGraphPin* OtherNativePin = Other->GetNativePin();
+	if(!SelfNativePin || !OtherNativePin)
+	{
+		return false;
+	}
+
+	const UEdGraphSchema_K2* Schema = GetDefault<UEdGraphSchema_K2>();
+	return Schema->TryCreateConnection(SelfNativePin, OtherNativePin);
+}
+
+bool UBlueprintGraphPin::CanCreateConnection(const UBlueprintGraphPin* Other) const
+{
+	const UEdGraphPin* OtherNativePin = GetNativePinSafe(Other);
+	const UEdGraphPin* NativePin = GetNativePin();
+	if(!OtherNativePin || !NativePin)
+	{
+		return false;
+	}
+	
+	const UEdGraphSchema_K2* Schema = GetDefault<UEdGraphSchema_K2>();
+	return Schema->CanCreateConnection(OtherNativePin, NativePin).Response != CONNECT_RESPONSE_DISALLOW;
+}
+	
+bool UBlueprintGraphPin::BreakPinLinks()
+{
+	UEdGraphPin* SelfNativePin = GetNativePin();
+	if(!SelfNativePin)
+	{
+		return false;
+	}
+	
+	const UEdGraphSchema_K2* Schema = GetDefault<UEdGraphSchema_K2>();
+	Schema->BreakPinLinks(*SelfNativePin, true);
+	return true;
+}
+
+bool UBlueprintGraphPin::BreakSinglePinLink(UBlueprintGraphPin* Other)
+{
+	if(!Other)
+	{
+		return false;
+	}
+
+	UEdGraphPin* SelfNativePin = GetNativePin();
+	UEdGraphPin* OtherNativePin = Other->GetNativePin();
+	if(!SelfNativePin || !OtherNativePin)
+	{
+		return false;
+	}
+
+	const UEdGraphSchema_K2* Schema = GetDefault<UEdGraphSchema_K2>();
+	Schema->BreakSinglePinLink(SelfNativePin, OtherNativePin);
+	return true;
+}
+
+FName UBlueprintGraphPin::GetPinName() const
+{
+	UEdGraphPin* Pin = GetNativePin();
+	if(!Pin)
+	{
+		return TEXT("");
+	}
+
+	return Pin->PinName;
+}
+
+FString UBlueprintGraphPin::GetPinTypeAsJsonSchema() const
+{
+	UEdGraphPin* Pin = GetNativePin();
+	if(!Pin)
+	{
+		return TEXT("");
+	}
+
+	UK2Node* OwningNode = Cast<UK2Node>(Pin->GetOwningNode());
+	if(!OwningNode)
+	{
+		return TEXT("");
+	}
+
+	return UBlueprintEditorLibrary::PinTypeToJsonSchema(Pin->PinType, OwningNode->GetBlueprintClassFromNode());
+}
+
+FEdGraphPinType UBlueprintGraphPin::GetPinType() const
+{
+	UEdGraphPin* Pin = GetNativePin();
+	if(!Pin)
+	{
+		return FEdGraphPinType();
+	}
+
+	return Pin->PinType;
+}
+
+TEnumAsByte<EEdGraphPinDirection> UBlueprintGraphPin::GetPinDirection() const
+{
+	return PinDirection;
+}
+
+FString UBlueprintGraphPin::GetPinValue() const
+{
+	UEdGraphPin* Pin = GetNativePin();
+	if(!Pin)
+	{
+		return TEXT("");
+	}
+
+	return Pin->GetDefaultAsString();
+}
+
+UK2Node* UBlueprintGraphPin::GetOwningNode() const
+{
+	return Node;
+}
+
+UEdGraphPin* UBlueprintGraphPin::GetNativePin() const
+{
+	if(!::IsValid(Node))
+	{
+		return nullptr;
+	}
+
+	return Node->FindPin(PinName, PinDirection);
+}
+
+UEdGraphPin* UBlueprintGraphPin::GetNativePinSafe(const UBlueprintGraphPin* Pin)
+{
+	if(!::IsValid(Pin))
+	{
+		return nullptr;
+	}
+
+	return Pin->GetNativePin();
+}
diff --git a/Engine/Source/Editor/BlueprintEditorLibrary/Private/BlueprintEditorLibrary/BlueprintGraphPin.h b/Engine/Source/Editor/BlueprintEditorLibrary/Private/BlueprintEditorLibrary/BlueprintGraphPin.h
new file mode 100644
--- /dev/null
+++ b/Engine/Source/Editor/BlueprintEditorLibrary/Private/BlueprintEditorLibrary/BlueprintGraphPin.h
@@ -0,0 +1,108 @@
+// Copyright Epic Games, Inc. All Rights Reserved.
+
+#pragma once
+
+#include "UObject/Object.h"
+#include "EdGraph/EdGraphNode.h"
+
+#include "BlueprintGraphPin.generated.h"
+
+class UEdGraph;
+class UBlueprint;
+
+class UK2Node;
+
+/** 
+ * Helper for interfacing with UEdGraphPin (not a UObject, relatively low level) as it's used in blueprints. This
+ * is a transient class and as such is created on demand and does not have useful identity, so comparison of reference
+ * equality to other UBlueprintGraphPins will not be meaningful. Wrapping UEdGraphPin makes this type less prone to
+ * leaking and safer to use across graph edit operations.
+ */
+UCLASS(Transient, BlueprintType, HideDropDown)
+class UBlueprintGraphPin : public UObject
+{
+	GENERATED_BODY()
+public:
+	/** Creates a UBlueprintGraphPin wrappin the underlying native Pin */
+	static UBlueprintGraphPin* FromNativePin(UEdGraphPin* Pin);
+
+	/**
+	 * Sets the literal value associated with this pin, returning false if the value is not valid
+	 */
+	UFUNCTION(BlueprintCallable, Category = "Blueprint Editor")
+	bool SetPinValue(const FString& Value);
+	
+	/**
+	 * Tries to create a connection between this pin and Other, returning false if the connection is not allowed
+	 */
+	UFUNCTION(BlueprintCallable, Category = "Blueprint Editor")
+	bool TryCreateConnection(UBlueprintGraphPin* Other);
+	
+	/**
+	 * Returns true if this and Other pin can be connected. Performing connection may result in implicitly created conversion nodes.
+	 */
+	UFUNCTION(BlueprintCallable, Category = "Blueprint Editor")
+	bool CanCreateConnection(const UBlueprintGraphPin* Other) const;
+
+	/**
+	 * Breaks all pin links associated with this pin
+	 */
+	UFUNCTION(BlueprintCallable, Category = "Blueprint Editor")
+	bool BreakPinLinks();
+	
+	/**
+	 * Breaks a pin link between this pin and Other
+	 */
+	UFUNCTION(BlueprintCallable, Category = "Blueprint Editor")
+	bool BreakSinglePinLink(UBlueprintGraphPin* Other);
+	
+	/**
+	 * Returns the low level name of the pin
+	 */
+	UFUNCTION(BlueprintCallable, Category = "Blueprint Editor")
+	FName GetPinName() const;
+	
+	/**
+	 * Returns the direction of the pin (either EGPD_Output or EGPD_Input)
+	 */
+	UFUNCTION(BlueprintCallable, Category = "Blueprint Editor")
+	TEnumAsByte<EEdGraphPinDirection> GetPinDirection() const;
+
+	/**
+	 * Returns a json schema encoded description of the pin's type:
+	 */
+	UFUNCTION(BlueprintCallable, Category = "Blueprint Editor")
+	FString GetPinTypeAsJsonSchema() const;
+
+	/**
+	 * Returns the pin's type, FEdGraphPinType is mostly opaque to script but can
+	 * be used e.g. to create variables of the same type
+	 */
+	UFUNCTION(BlueprintCallable, Category = "Blueprint Editor")
+	FEdGraphPinType GetPinType() const;
+
+	/**
+	 * Returns the literal value associated with the pin. Empty string for pins 
+	 * that have nothing set or are connected to some other pin:
+	 */
+	UFUNCTION(BlueprintCallable, Category = "Blueprint Editor")
+	FString GetPinValue() const;
+	
+	/**
+	 * Returns the pin's node
+	 */
+	UFUNCTION(BlueprintCallable, Category = "Blueprint Editor")
+	UK2Node* GetOwningNode() const;
+
+	UEdGraphPin* GetNativePin() const;
+	static UEdGraphPin* GetNativePinSafe(const UBlueprintGraphPin* Pin);
+private:
+	UPROPERTY()
+	TObjectPtr<UK2Node> Node;
+
+	UPROPERTY()
+	TEnumAsByte<EEdGraphPinDirection> PinDirection;
+
+	UPROPERTY()
+	FName PinName;
+};
diff --git a/Engine/Source/Editor/BlueprintEditorLibrary/Public/BlueprintEditorLibrary.h b/Engine/Source/Editor/BlueprintEditorLibrary/Public/BlueprintEditorLibrary.h
--- a/Engine/Source/Editor/BlueprintEditorLibrary/Public/BlueprintEditorLibrary.h
+++ b/Engine/Source/Editor/BlueprintEditorLibrary/Public/BlueprintEditorLibrary.h
@@ -7,6 +7,7 @@
 #include "Containers/UnrealString.h"
 #include "CoreMinimal.h"
 #include "EdGraph/EdGraphNode.h"
+#include "EdGraph/EdGraphPin.h"
 #include "HAL/Platform.h"
 #include "Kismet/BlueprintFunctionLibrary.h"
 #include "Logging/LogMacros.h"
@@ -48,7 +49,6 @@ class UBlueprintEditorLibrary : public UBlueprintFunctionLibrary
 {
 	GENERATED_UCLASS_BODY()
 public:
-
 	/**
 	* Replace any references of variables with the OldVarName to references of those with the NewVarName if possible
 	*
@@ -304,6 +304,12 @@ public:
 	 UFUNCTION(BlueprintCallable, Category = "Blueprint Editor")
 	 static UE_API UBlueprint* CreateBlueprintAssetWithParent(const FString& AssetPath, UClass* ParentClass);
 
+	/**
+	 * Lists the names of the editable graphs in the BP
+	 */
+	UFUNCTION(BlueprintCallable, Category = "Blueprint Editor")
+	static void ListGraphNames(const UBlueprint* Blueprint, TArray<FName>& OutGraphNames);
+
 	/**
 	  * Adds a member variable to the specified blueprint inferring the type from a provided value.
 	  * 
@@ -349,6 +355,100 @@ public:
 	/** @return a map of KeyType to ValueType type - returns 'int' type if invalid type is provided */
 	UFUNCTION(BlueprintCallable, Category = "Blueprint Editor")
 	static UE_API FEdGraphPinType GetMapType(const FEdGraphPinType& KeyType,const FEdGraphPinType& ValueType);
+	
+	/**
+	 * @ A json schema string describing the type
+	 */
+	UFUNCTION(BlueprintCallable, Category = "Blueprint Editor", meta = (ScriptMethod))
+	static UE_API FString PinTypeToJsonSchema(const FEdGraphPinType& PinType, const UClass* SelfContext);
+
+	/// Begin UK2Node extension/script methods
+	
+	/**
+	 * Returns all visible pins on this node, optionally discriminated by direction
+	 */
+	UFUNCTION(BlueprintCallable, Category = "Blueprint Graph Editing", meta = (ScriptMethod))
+	static TArray<UBlueprintGraphPin*> ListAllPins(const UK2Node* Node, TEnumAsByte<EEdGraphPinDirection> InDirection = EEdGraphPinDirection::EGPD_MAX);
+	
+	/**
+	 * Returns an input pin specified by index and, optionally, type
+	 */
+	UFUNCTION(BlueprintCallable, Category = "Blueprint Graph Editing", meta = (ScriptMethod))
+	static UBlueprintGraphPin* FindInputPinByIndex(const UK2Node* Node, int32 Index, const FEdGraphPinType& Type = FEdGraphPinType() );
+	
+	/**
+	 * Returns an input pin specified by name and, optionally, type
+	 */
+	UFUNCTION(BlueprintCallable, Category = "Blueprint Graph Editing", meta = (ScriptMethod))
+	static UBlueprintGraphPin* FindInputPin(const UK2Node* Node, FName PinName, const FEdGraphPinType& Type = FEdGraphPinType() );
+	
+	/**
+	 * Returns an output pin specified by index and, optionally, type
+	 */
+	UFUNCTION(BlueprintCallable, Category = "Blueprint Graph Editing", meta = (ScriptMethod))
+	static UBlueprintGraphPin* FindOutputPinByIndex(const UK2Node* Node, int32 Index, const FEdGraphPinType& Type = FEdGraphPinType() );
+	
+	/**
+	 * Returns an output pin specified by name and, optionally, type
+	 */
+	UFUNCTION(BlueprintCallable, Category = "Blueprint Graph Editing", meta = (ScriptMethod))
+	static UBlueprintGraphPin* FindOutputPin(const UK2Node* Node, FName PinName, const FEdGraphPinType& Type = FEdGraphPinType() );
+	
+	/**
+	 * Returns the 'execute' or 'do' pin associated with this node, if any
+	 */
+	UFUNCTION(BlueprintCallable, Category = "Blueprint Graph Editing", meta = (ScriptMethod))
+	static UBlueprintGraphPin* FindExecutePin(const UK2Node* Node);
+	
+	/**
+	 * Returns the 'then' pin associated with this node, if any
+	 */
+	UFUNCTION(BlueprintCallable, Category = "Blueprint Graph Editing", meta = (ScriptMethod))
+	static UE_API UBlueprintGraphPin* FindThenPin(const UK2Node* Node);
+	
+	/**
+	 * Returns the self pin associated with this node, if any
+	 */
+	UFUNCTION(BlueprintCallable, Category = "Blueprint Graph Editing", meta = (ScriptMethod))
+	static UBlueprintGraphPin* FindSelfPin(const UK2Node* Node);
+	
+	/**
+	 * Returns the single data output pin associated with this node, returning null if there are multiple or no output(s)
+	 */
+	UFUNCTION(BlueprintCallable, Category = "Blueprint Graph Editing", meta = (ScriptMethod))
+	static UBlueprintGraphPin* FindResultPin(const UK2Node* Node);
+
+	/**
+	 * Returns the single data input pin associated with this node, returning null if there are multiple or no input(s)
+	 */
+	UFUNCTION(BlueprintCallable, Category = "Blueprint Graph Editing", meta = (ScriptMethod))
+	static UBlueprintGraphPin* FindDataInputPin(const UK2Node* Node);
+	
+	/**
+	 * Sets this node's visual position in the graph
+	 */
+	UFUNCTION(BlueprintCallable, Category = "Blueprint Graph Editing", meta = (ScriptMethod))
+	static void SetNodePos(UK2Node* Node, FIntPoint pos);
+	
+	/**
+	 * Returns this node's visual position in the graph
+	 */
+	UFUNCTION(BlueprintCallable, Category = "Blueprint Graph Editing", meta = (ScriptMethod))
+	static FIntPoint GetNodePos(const UK2Node* Node);
+
+	/// end UK2Node extension/script methods
+
+	/// begin UK2Node_IfThenElse textension/script methods:
+
+	/** Finds the Else pin on this node, wrapped in a UBlueprintGraphPin */
+	UFUNCTION(BlueprintCallable, Category = "Blueprint Graph Editing", meta = (ScriptMethod))
+	static UBlueprintGraphPin* FindElsePin(const UK2Node_IfThenElse* Node);
+	
+	/** Finds the Condition pin on this node, wrapped in a UBlueprintGraphPin */
+	UFUNCTION(BlueprintCallable, Category = "Blueprint Graph Editing", meta = (ScriptMethod))
+	static UBlueprintGraphPin* FindConditionPin(const UK2Node_IfThenElse* Node);
+
+	/// end UK2Node_IfThenElse extension/script methods
 };
 
 #undef UE_API
diff --git a/Engine/Source/Editor/BlueprintGraph/Classes/K2Node_EditablePinBase.h b/Engine/Source/Editor/BlueprintGraph/Classes/K2Node_EditablePinBase.h
--- a/Engine/Source/Editor/BlueprintGraph/Classes/K2Node_EditablePinBase.h
+++ b/Engine/Source/Editor/BlueprintGraph/Classes/K2Node_EditablePinBase.h
@@ -197,7 +197,7 @@ class UK2Node_EditablePinBase : public UK2Node
 	* @param	PinName name of pin check existence of
 	* @return	True if a user defined pin with this name exists
 	*/
-	bool UserDefinedPinExists(const FName PinName) const;
+	BLUEPRINTGRAPH_API bool UserDefinedPinExists(const FName PinName) const;
 
 	/**
 	 * Creates a new pin on the node from the specified user pin info.
