Tweak comment, removing some needless imports from the example and emphasizing BlueprintGraphPin and BlueprintEditorLibrary's relatedness in the comment body.

#rnx

#Synced-CL 49832464

diff --git a/Engine/Source/Editor/BlueprintEditorLibrary/Private/BlueprintEditorLibrary/BlueprintGraphEditor.h b/Engine/Source/Editor/BlueprintEditorLibrary/Private/BlueprintEditorLibrary/BlueprintGraphEditor.h
--- a/Engine/Source/Editor/BlueprintEditorLibrary/Private/BlueprintEditorLibrary/BlueprintGraphEditor.h
+++ b/Engine/Source/Editor/BlueprintEditorLibrary/Private/BlueprintEditorLibrary/BlueprintGraphEditor.h
@@ -26,15 +26,14 @@ class UBlueprintGraphPin;
 /**
  * Object oriented helper for editing graphs in a blueprint object
  * Instantiate the UBlueprintGraphEditor via CreateAndEditFunctionGraph
- * or GetGraphEditor
+ * or GetGraphEditor. Static helpers can be found in BlueprintEditorLibrary
+ * and pin specific routines are encapsulated by BlueprintGraphPin
  * 
  * Python Example:
 
 import unreal
 
 BlueprintGraphEditor = unreal.BlueprintGraphEditor
-BlueprintEditorLibrary = unreal.BlueprintEditorLibrary
-BlueprintGraphPin = unreal.BlueprintGraphPin
 
 print_string = "/Script/Engine.KismetSystemLibrary.PrintString"
 
@@ -44,7 +43,7 @@ def GenHelloWorld(bp):
     start_pos = start_pin.get_owning_node().get_node_pos()
     print_node = editor.add_call_function_node(print_string)
     print_node.find_input_pin("InString").set_pin_value("Generated Script")
-    editor.find_graph_entry_pin().try_create_connection(print_node.find_execute_pin())
+    start_pin.try_create_connection(print_node.find_execute_pin())
     print_node.set_node_pos(unreal.IntPoint(start_pos.x + 160, start_pos.y))
 
  */
