Make GetBlueprintAsset a little smarter so that it's easier to find a blueprint object associated with a package or asset

[RN] UBlueprintEditorLibrary::GetBlueprintAsset now looks for adjacent UBlueprint objects, e.g. if a Level asset is provided as the object to get the UBlueprint from

#Synced-CL 50069009

diff --git a/Engine/Source/Editor/BlueprintEditorLibrary/Private/BlueprintEditorLibrary.cpp b/Engine/Source/Editor/BlueprintEditorLibrary/Private/BlueprintEditorLibrary.cpp
--- a/Engine/Source/Editor/BlueprintEditorLibrary/Private/BlueprintEditorLibrary.cpp
+++ b/Engine/Source/Editor/BlueprintEditorLibrary/Private/BlueprintEditorLibrary.cpp
@@ -658,7 +658,48 @@ void UBlueprintEditorLibrary::RenameGraph(UEdGraph* Graph, const FString& NewNam
 
 UBlueprint* UBlueprintEditorLibrary::GetBlueprintAsset(UObject* Object)
 {
-	return Cast<UBlueprint>(Object);
+	if(UBlueprint* BP = Cast<UBlueprint>(Object))
+	{
+		return BP;
+	}
+	else if(UBlueprintGeneratedClass* Class = Cast<UBlueprintGeneratedClass>(Object))
+	{
+		return Cast<UBlueprint>(Class->ClassGeneratedBy);
+	}
+
+	UPackage* Package = nullptr;
+	if(Object->IsAsset() && Object->GetExternalPackage() == nullptr)
+	{
+		Package = Object->GetPackage();
+	}
+	else
+	{
+		Package = Cast<UPackage>(Object);
+	}
+
+	if(Package)
+	{
+		// search for a public root level blueprint:
+		TArray<UObject*> RootObjects;
+		GetObjectsWithOuter(Package, RootObjects, false, RF_Transient, EInternalObjectFlags::Garbage);
+		for(UObject* Obj : RootObjects)
+		{
+			if(!Obj->HasAnyFlags(RF_Public))
+			{
+				continue;
+			}
+			
+			if(UBlueprint* InnerBP = Cast<UBlueprint>(Obj))
+			{
+				return InnerBP;
+			}
+			else if(UBlueprintGeneratedClass* Class = Cast<UBlueprintGeneratedClass>(Obj))
+			{
+				return Cast<UBlueprint>(Class->ClassGeneratedBy);
+			}
+		}
+	}
+	return nullptr;
 }
 
 UBlueprint* UBlueprintEditorLibrary::GetBlueprintForClass(UClass* Class, bool& bDoesClassHaveBlueprint)
diff --git a/Engine/Source/Editor/BlueprintEditorLibrary/Public/BlueprintEditorLibrary.h b/Engine/Source/Editor/BlueprintEditorLibrary/Public/BlueprintEditorLibrary.h
--- a/Engine/Source/Editor/BlueprintEditorLibrary/Public/BlueprintEditorLibrary.h
+++ b/Engine/Source/Editor/BlueprintEditorLibrary/Public/BlueprintEditorLibrary.h
@@ -194,12 +194,15 @@ public:
 	static UE_API void RenameGraph(UEdGraph* Graph, const FString& NewNameStr = FString(TEXT("NewGraph")));
 
 	/**
-	* Casts the provided Object to a Blueprint - the root asset type of a blueprint asset. Note
-	* that the blueprint asset itself is editor only and not present in cooked assets.
+	* Finds the UBlueprint associated with the object, locally searching the object graph for a
+	* UBlueprint associated with an asset object. If the Object is a UBlueprint this function will 
+	* perform a simple cast. Note that the blueprint object itself is editor only and not present 
+	* in cooked assets.
 	*
 	* @param Object			The object we need to get the UBlueprint from
 	*
-	* @return UBlueprint*	The blueprint type of the given object, nullptr if the object is not a blueprint.
+	* @return UBlueprint*	The blueprint associated with the given object, nullptr if the object has 
+	* no associated blueprint or the blueprint is in another package
 	*/
 	UFUNCTION(BlueprintCallable, Category = "Blueprint Upgrade Tools", meta = (Keywords = "cast"))
 	static UE_API UBlueprint* GetBlueprintAsset(UObject* Object);
