Add missed null check

#rnx

diff --git a/Engine/Source/Editor/BlueprintEditorLibrary/Private/BlueprintEditorLibrary.cpp b/Engine/Source/Editor/BlueprintEditorLibrary/Private/BlueprintEditorLibrary.cpp
--- a/Engine/Source/Editor/BlueprintEditorLibrary/Private/BlueprintEditorLibrary.cpp
+++ b/Engine/Source/Editor/BlueprintEditorLibrary/Private/BlueprintEditorLibrary.cpp
@@ -658,6 +658,11 @@ void UBlueprintEditorLibrary::RenameGraph(UEdGraph* Graph, const FString& NewNam
 
 UBlueprint* UBlueprintEditorLibrary::GetBlueprintAsset(UObject* Object)
 {
+	if(Object == nullptr)
+	{
+		return nullptr;
+	}
+
 	if(UBlueprint* BP = Cast<UBlueprint>(Object))
 	{
 		return BP;
