Added a ListGraphs method to BlueprintEditorLibrary.

#Synced-CL 50146075

diff --git a/Engine/Source/Editor/BlueprintEditorLibrary/Private/BlueprintEditorLibrary.cpp b/Engine/Source/Editor/BlueprintEditorLibrary/Private/BlueprintEditorLibrary.cpp
--- a/Engine/Source/Editor/BlueprintEditorLibrary/Private/BlueprintEditorLibrary.cpp
+++ b/Engine/Source/Editor/BlueprintEditorLibrary/Private/BlueprintEditorLibrary.cpp
@@ -332,6 +332,22 @@ UEdGraph* UBlueprintEditorLibrary::FindGraph(UBlueprint* Blueprint, FName GraphN
 	return nullptr;
 }
 
+TArray<FName> UBlueprintEditorLibrary::ListGraphs(UBlueprint* Blueprint)
+{
+	TArray<FName> GraphNames;
+	if (Blueprint)
+	{
+		TArray<UEdGraph*> AllGraphs;
+		Blueprint->GetAllGraphs(AllGraphs);
+
+		for (UEdGraph* CurrentGraph : AllGraphs)
+		{
+			GraphNames.Add(CurrentGraph->GetFName());
+		}
+	}
+	return GraphNames;
+}
+
 UK2Node_PromotableOperator* CreateOpNode(const FName OpName, UEdGraph* Graph, const int32 AdditionalPins)
 {
 	if (!Graph)
diff --git a/Engine/Source/Editor/BlueprintEditorLibrary/Public/BlueprintEditorLibrary.h b/Engine/Source/Editor/BlueprintEditorLibrary/Public/BlueprintEditorLibrary.h
--- a/Engine/Source/Editor/BlueprintEditorLibrary/Public/BlueprintEditorLibrary.h
+++ b/Engine/Source/Editor/BlueprintEditorLibrary/Public/BlueprintEditorLibrary.h
@@ -129,6 +129,16 @@ public:
 	UFUNCTION(BlueprintCallable, Category = "Blueprint Upgrade Tools")
 	static UE_API UEdGraph* FindGraph(UBlueprint* Blueprint, FName GraphName);
 
+	/**
+	* Lists the names of the graphs that a Blueprint contains.
+	*
+	* @param Blueprint			Blueprint to enumerate	
+	*
+	* @return TArray<FName>		A list of the names of the graphs in the blueprint.
+	*/
+	UFUNCTION(BlueprintCallable, Category = "Blueprint Upgrade Tools")
+	static UE_API TArray<FName> ListGraphs(UBlueprint* Blueprint);
+
 	/**
 	* Replace any old operator nodes (float + float, vector + float, int + vector, etc)
 	* with the newer Promotable Operator version of the node. Preserve any connections the
