Convert BlueprintEditorLibrary::ListGraphs to return a list of UEdGraph*.

#Synced-CL 50146075

diff --git a/Engine/Source/Editor/BlueprintEditorLibrary/Private/BlueprintEditorLibrary.cpp b/Engine/Source/Editor/BlueprintEditorLibrary/Private/BlueprintEditorLibrary.cpp
--- a/Engine/Source/Editor/BlueprintEditorLibrary/Private/BlueprintEditorLibrary.cpp
+++ b/Engine/Source/Editor/BlueprintEditorLibrary/Private/BlueprintEditorLibrary.cpp
@@ -332,20 +332,24 @@ UEdGraph* UBlueprintEditorLibrary::FindGraph(UBlueprint* Blueprint, FName GraphN
 	return nullptr;
 }
 
-TArray<FName> UBlueprintEditorLibrary::ListGraphs(UBlueprint* Blueprint)
+TArray<UEdGraph*> UBlueprintEditorLibrary::ListGraphs(UBlueprint* Blueprint)
 {
-	TArray<FName> GraphNames;
+	TArray<UEdGraph*> Graphs;
 	if (Blueprint)
 	{
 		TArray<UEdGraph*> AllGraphs;
 		Blueprint->GetAllGraphs(AllGraphs);
 
-		for (UEdGraph* CurrentGraph : AllGraphs)
+		for (UEdGraph* Graph : AllGraphs)
 		{
-			GraphNames.Add(CurrentGraph->GetFName());
+			if (!IsValid(Graph) || Graph->HasAnyFlags(RF_Transient))
+			{
+				continue;
+			}
+			Graphs.Add(Graph);
 		}
 	}
-	return GraphNames;
+	return Graphs;
 }
 
 UK2Node_PromotableOperator* CreateOpNode(const FName OpName, UEdGraph* Graph, const int32 AdditionalPins)
diff --git a/Engine/Source/Editor/BlueprintEditorLibrary/Public/BlueprintEditorLibrary.h b/Engine/Source/Editor/BlueprintEditorLibrary/Public/BlueprintEditorLibrary.h
--- a/Engine/Source/Editor/BlueprintEditorLibrary/Public/BlueprintEditorLibrary.h
+++ b/Engine/Source/Editor/BlueprintEditorLibrary/Public/BlueprintEditorLibrary.h
@@ -130,14 +130,14 @@ public:
 	static UE_API UEdGraph* FindGraph(UBlueprint* Blueprint, FName GraphName);
 
 	/**
-	* Lists the names of the graphs that a Blueprint contains.
+	* Lists all of the graphs that a Blueprint contains.
 	*
 	* @param Blueprint			Blueprint to enumerate	
 	*
-	* @return TArray<FName>		A list of the names of the graphs in the blueprint.
+	* @return TArray<FName>		A list of the graphs in the blueprint.
 	*/
 	UFUNCTION(BlueprintCallable, Category = "Blueprint Upgrade Tools")
-	static UE_API TArray<FName> ListGraphs(UBlueprint* Blueprint);
+	static UE_API TArray<UEdGraph*> ListGraphs(UBlueprint* Blueprint);
 
 	/**
 	* Replace any old operator nodes (float + float, vector + float, int + vector, etc)
