Make blueprint graph editor more consistent with use of FName and UEdGraph

#rnx

#Synced-CL 50069009

diff --git a/Engine/Source/Editor/BlueprintEditorLibrary/Private/BlueprintEditorLibrary/BlueprintGraphEditor.cpp b/Engine/Source/Editor/BlueprintEditorLibrary/Private/BlueprintEditorLibrary/BlueprintGraphEditor.cpp
--- a/Engine/Source/Editor/BlueprintEditorLibrary/Private/BlueprintEditorLibrary/BlueprintGraphEditor.cpp
+++ b/Engine/Source/Editor/BlueprintEditorLibrary/Private/BlueprintEditorLibrary/BlueprintGraphEditor.cpp
@@ -58,7 +58,7 @@ UBlueprintGraphEditor* UBlueprintGraphEditor::CreateAndEditFunctionGraph(UBluepr
 	return GraphEditor;
 }
 
-UBlueprintGraphEditor* UBlueprintGraphEditor::GetGraphEditor(UBlueprint* InBlueprint, FName GraphName)
+UBlueprintGraphEditor* UBlueprintGraphEditor::GetGraphEditorByName(UBlueprint* InBlueprint, FName GraphName)
 {
 	const auto FindGraphByName = [](const TArray<UEdGraph*>& GraphList, FName Name)
 	{
@@ -87,6 +87,26 @@ UBlueprintGraphEditor* UBlueprintGraphEditor::GetGraphEditor(UBlueprint* InBluep
 	return GraphEditor;
 }
 
+UBlueprintGraphEditor* UBlueprintGraphEditor::GetGraphEditor(UEdGraph* InGraph)
+{
+	if(!InGraph)
+	{
+		UE_LOGFMT(LogBlueprint, Error, "Cannot make a graph editor for a null graph");
+		return nullptr;
+	}
+
+	UBlueprint* BP = FBlueprintEditorUtils::FindBlueprintForGraph(InGraph);
+	if(!BP)
+	{
+		UE_LOGFMT(LogBlueprint, Error, "Provided graph is not part of a blueprint: {GraphName}", InGraph->GetPathName());
+	}
+	
+	UBlueprintGraphEditor* GraphEditor = NewObject<UBlueprintGraphEditor>();
+	GraphEditor->Blueprint = BP;
+	GraphEditor->Graph = InGraph;
+	return GraphEditor;
+}
+
 void UBlueprintGraphEditor::ListAllNodes(TArray<UK2Node*>& OutNodes)
 {
 	if(!IsValid())
diff --git a/Engine/Source/Editor/BlueprintEditorLibrary/Private/BlueprintEditorLibrary/BlueprintGraphEditor.h b/Engine/Source/Editor/BlueprintEditorLibrary/Private/BlueprintEditorLibrary/BlueprintGraphEditor.h
--- a/Engine/Source/Editor/BlueprintEditorLibrary/Private/BlueprintEditorLibrary/BlueprintGraphEditor.h
+++ b/Engine/Source/Editor/BlueprintEditorLibrary/Private/BlueprintEditorLibrary/BlueprintGraphEditor.h
@@ -63,8 +63,12 @@ public:
 	
 	/** returns a new UBlueprintGraphEditor for a graph in InBlueprint named GraphName - useful for editing graphs that already exist */
 	UFUNCTION(BlueprintCallable, Category = "Blueprint Editor")
-	static UBlueprintGraphEditor* GetGraphEditor(UBlueprint* InBlueprint, FName GraphName);
-	
+	static UBlueprintGraphEditor* GetGraphEditorByName(UBlueprint* InBlueprint, FName GraphName);
+
+	/** returns a new UBlueprintGraphEditor for the provided graph - useful for editing graphs that already exist */
+	UFUNCTION(BlueprintCallable, Category = "Blueprint Editor")
+	static UBlueprintGraphEditor* GetGraphEditor(UEdGraph* InGraph);
+
 	/**
 	 * Lists all nodes in the graph
 	 */
