Blueprint editor scripting:
* AddCustomEventNode
* FindEventNode
* GetNodeTitle
* GetPinTypeDisplayString
* ListConnectedPins
* IsSameNativePin for pin identity comparison

#rnx
#Synced-CL 50763490

diff --git a/Engine/Source/Editor/BlueprintEditorLibrary/Private/BlueprintEditorLibrary.cpp b/Engine/Source/Editor/BlueprintEditorLibrary/Private/BlueprintEditorLibrary.cpp
--- a/Engine/Source/Editor/BlueprintEditorLibrary/Private/BlueprintEditorLibrary.cpp
+++ b/Engine/Source/Editor/BlueprintEditorLibrary/Private/BlueprintEditorLibrary.cpp
@@ -1686,6 +1686,16 @@ FIntPoint UBlueprintEditorLibrary::GetNodePos(const UK2Node* Node)
 	return {Node->NodePosX, Node->NodePosY};
 }
 
+FString UBlueprintEditorLibrary::GetNodeTitle(const UK2Node* Node)
+{
+	if(!::IsValid(Node))
+	{
+		return TEXT("");
+	}
+
+	return Node->GetNodeTitle(ENodeTitleType::ListView).ToString();
+}
+
 UBlueprintGraphPin* UBlueprintEditorLibrary::FindElsePin(const UK2Node_IfThenElse* Node)
 {
 	if(!::IsValid(Node))
diff --git a/Engine/Source/Editor/BlueprintEditorLibrary/Private/BlueprintEditorLibrary/BlueprintGraphEditor.cpp b/Engine/Source/Editor/BlueprintEditorLibrary/Private/BlueprintEditorLibrary/BlueprintGraphEditor.cpp
--- a/Engine/Source/Editor/BlueprintEditorLibrary/Private/BlueprintEditorLibrary/BlueprintGraphEditor.cpp
+++ b/Engine/Source/Editor/BlueprintEditorLibrary/Private/BlueprintEditorLibrary/BlueprintGraphEditor.cpp
@@ -10,6 +10,8 @@
 #include "Blueprint/BlueprintExceptionInfo.h"
 #include "Kismet/KismetSystemLibrary.h"
 #include "K2Node_CallFunction.h"
+#include "K2Node_CustomEvent.h"
+#include "K2Node_Event.h"
 #include "K2Node_FunctionEntry.h"
 #include "K2Node_FunctionResult.h"
 #include "K2Node_MacroInstance.h"
@@ -19,6 +21,7 @@
 #include "K2Node_VariableGet.h"
 #include "K2Node_VariableSet.h"
 #include "Kismet2/BlueprintEditorUtils.h"
+#include "Kismet2/Kismet2NameValidators.h"
 #include "Logging/StructuredLog.h"
 #include "UObject/Class.h"
 #include "UObject/SoftObjectPtr.h"
@@ -221,6 +224,46 @@ UBlueprintGraphPin* UBlueprintGraphEditor::FindGraphEntryPin() const
 	return UBlueprintEditorLibrary::FindThenPin(Entry);//Entry->FindThenPin();
 }
 
+UK2Node_Event* UBlueprintGraphEditor::FindEventNode(FName EventName) const
+{
+	if(!IsValid())
+	{
+		return nullptr;
+	}
+
+	if (!FBlueprintEditorUtils::IsEventGraph(Graph))
+	{
+		UE_LOGFMT(LogBlueprint, Warning, "Non event graph {Graph} does not allow event editing", Graph->GetPathName());
+		return nullptr;
+	}
+
+	TArray<UK2Node_Event*> EventNodes;
+	Graph->GetNodesOfClass(EventNodes);
+
+	UK2Node_Event* Result = nullptr;
+	for(UK2Node_Event* EventNode : EventNodes)
+	{
+		if(EventNode->GetFunctionName() == EventName)
+		{
+			if(!Result)
+			{
+				Result = EventNode;
+			}
+			else
+			{
+				UE_LOGFMT(LogBlueprint, Warning, "Found ambiguous event node {EventName} for editing in {Graph}", EventName, Graph->GetPathName());
+			}
+		}
+	}
+
+	if(Result == nullptr)
+	{
+		UE_LOGFMT(LogBlueprint, Warning, "No event node {EventName} for in {Graph}", EventName, Graph->GetPathName());
+	}
+
+	return Result;
+}
+
 UK2Node_FunctionResult* UBlueprintGraphEditor::AddGraphOutputParameter(FName OutputName, const FEdGraphPinType& PinType)
 {
 	if(!IsValid())
@@ -809,6 +852,58 @@ UK2Node_VariableSet* UBlueprintGraphEditor::AddSetLocalVariableNode(FName LocalN
 	return Result;
 }
 
+UK2Node_CustomEvent* UBlueprintGraphEditor::AddCustomEventNode(const FString& EventName)
+{
+	if(!IsValid())
+	{
+		return nullptr;
+	}
+
+	if (!FBlueprintEditorUtils::IsEventGraph(Graph))
+	{
+		UE_LOGFMT(LogBlueprint, Warning, "Cannot add custom event to {Graph} - it is not an event graph", Graph->GetPathName());
+		return nullptr;
+	}
+
+	UK2Node_CustomEvent* NewNode = NewObject<UK2Node_CustomEvent>(Graph);
+	check(NewNode != nullptr);
+	NewNode->CreateNewGuid();
+
+	// position the node before invoking PostSpawnDelegate (in case it 
+	// wishes to modify this positioning)
+	const FVector2D Location = GetNodeLocation();
+	NewNode->NodePosX = static_cast<int32>(Location.X);
+	NewNode->NodePosY = static_cast<int32>(Location.Y);
+
+	NewNode->CustomFunctionName =  FBlueprintEditorUtils::FindUniqueCustomEventName(Blueprint);
+	NewNode->bIsEditable = true;
+	NewNode->SetFlags(RF_Transactional);
+	NewNode->AllocateDefaultPins();
+	NewNode->PostPlacedNewNode();
+
+	Graph->Modify();
+	// the FBlueprintMenuActionItem should do the selecting
+	Graph->AddNode(NewNode, /*bFromUI =*/true, /*bSelectNewNode =*/false);
+
+	// Unfortunately we can't validate the name easilly without first spawning the node, but we can honor
+	// name validity nonetheless:
+	TSharedPtr<INameValidatorInterface> NameValidator = NewNode->MakeNameValidator();
+	EValidatorResult ValidatorResult = NameValidator->IsValid(EventName);
+	if(ValidatorResult != EValidatorResult::Ok)
+	{
+		UE_LOGFMT(LogBlueprint, Warning, "User provided name was invalid {ErrorMessage} - node named {FallbackName}", 
+			INameValidatorInterface::GetErrorText(EventName, ValidatorResult).ToString(), NewNode->CustomFunctionName);
+	}
+	else
+	{
+		NewNode->CustomFunctionName = FName(*EventName);
+	}
+
+	HandleNodeSpawned(NewNode);
+
+	return NewNode;
+}
+
 void UBlueprintGraphEditor::RemoveNodes(const TArray<UK2Node*>& Nodes)
 {
 	if(!IsValid())
@@ -911,6 +1006,21 @@ UK2Node_FunctionEntry* UBlueprintGraphEditor::GetFunctionEntryNode() const
 void UBlueprintGraphEditor::HandleNodeSpawned(UK2Node* Result)
 {
 	UE_LOGFMT(LogBlueprint, Display, "Node {node_path} spawn at ({pos_x},{pox_y})", Result->GetPathName(), Result->NodePosX, Result->NodePosY);
+
+	if (Result->GetSchema()->MarkBlueprintDirtyFromNewNode(Blueprint, Result))
+	{
+		return;
+	}
+
+	UK2Node* K2Node = Cast<UK2Node>(Result);
+	if (Result->NodeCausesStructuralBlueprintChange())
+	{
+		FBlueprintEditorUtils::MarkBlueprintAsStructurallyModified(Blueprint);
+	}
+	else
+	{
+		FBlueprintEditorUtils::MarkBlueprintAsModified(Blueprint);
+	}
 }
 
 const FProperty* UBlueprintGraphEditor::FindParameterByName(FName Name) const
diff --git a/Engine/Source/Editor/BlueprintEditorLibrary/Private/BlueprintEditorLibrary/BlueprintGraphEditor.h b/Engine/Source/Editor/BlueprintEditorLibrary/Private/BlueprintEditorLibrary/BlueprintGraphEditor.h
--- a/Engine/Source/Editor/BlueprintEditorLibrary/Private/BlueprintEditorLibrary/BlueprintGraphEditor.h
+++ b/Engine/Source/Editor/BlueprintEditorLibrary/Private/BlueprintEditorLibrary/BlueprintGraphEditor.h
@@ -113,6 +113,15 @@ public:
 	 */
 	UFUNCTION(BlueprintCallable, Category = "Blueprint Editor")
 	UBlueprintGraphPin* FindGraphEntryPin() const;
+
+	/**
+	 * Returns the event node associated with EventName - aka Member Name
+	 * of the event. Names will not resolve against the display name, so e.g. 'Begin Play'
+	 * is not a valid member name - instead search for ReceiveBeginPlay, the underlying
+	 * BlueprintImplementableEvent
+	 */
+	UFUNCTION(BlueprintCallable, Category = "Blueprint Editor")
+	UK2Node_Event* FindEventNode(FName EventName) const;
 	
 	/**
 	 * Adds a graph output parameter, returning the UK2Node_FunctionResult that has the new output on it.
@@ -254,6 +263,12 @@ public:
 	UFUNCTION(BlueprintCallable, Category = "Blueprint Editor")
 	UK2Node_VariableSet* AddSetLocalVariableNode(FName LocalName);
 
+	/**
+	 * Creates a UK2Node_CustomEvent node if the current graph is an event graph, using the provided EventName
+	 */
+	UFUNCTION(BlueprintCallable, Category = "Blueprint Editor")
+	UK2Node_CustomEvent* AddCustomEventNode(const FString& EventName);
+
 	/**
 	 * Removes (aka deletes) the provided nodes from the current graph
 	 */
diff --git a/Engine/Source/Editor/BlueprintEditorLibrary/Private/BlueprintEditorLibrary/BlueprintGraphPin.cpp b/Engine/Source/Editor/BlueprintEditorLibrary/Private/BlueprintEditorLibrary/BlueprintGraphPin.cpp
--- a/Engine/Source/Editor/BlueprintEditorLibrary/Private/BlueprintEditorLibrary/BlueprintGraphPin.cpp
+++ b/Engine/Source/Editor/BlueprintEditorLibrary/Private/BlueprintEditorLibrary/BlueprintGraphPin.cpp
@@ -27,6 +27,18 @@ UBlueprintGraphPin* UBlueprintGraphPin::FromNativePin(UEdGraphPin* Pin)
 	return ScriptPin;
 }
 
+bool UBlueprintGraphPin::IsSameNativePin(const UBlueprintGraphPin* Other) const
+{
+	const UEdGraphPin* OtherNativePin = GetNativePinSafe(Other);
+	const UEdGraphPin* NativePin = GetNativePin();
+	if(!OtherNativePin || !NativePin)
+	{
+		return false;
+	}
+
+	return OtherNativePin == NativePin;
+}
+
 bool UBlueprintGraphPin::SetPinValue(const FString& Value)
 {						 
 	const UEdGraphSchema_K2* Schema = GetDefault<UEdGraphSchema_K2>();
@@ -136,6 +148,18 @@ FString UBlueprintGraphPin::GetPinTypeAsJsonSchema() const
 	return UBlueprintEditorLibrary::PinTypeToJsonSchema(Pin->PinType, OwningNode->GetBlueprintClassFromNode());
 }
 
+FText UBlueprintGraphPin::GetPinTypeDisplayString() const
+{
+	UEdGraphPin* Pin = GetNativePin();
+	if(!Pin)
+	{
+		return FText();
+	}
+
+	const UEdGraphSchema_K2* Schema = GetDefault<UEdGraphSchema_K2>();
+	return Schema->TypeToText(Pin->PinType);
+}
+
 FEdGraphPinType UBlueprintGraphPin::GetPinType() const
 {
 	UEdGraphPin* Pin = GetNativePin();
@@ -168,6 +192,26 @@ UK2Node* UBlueprintGraphPin::GetOwningNode() const
 	return Node;
 }
 
+TArray<UBlueprintGraphPin*> UBlueprintGraphPin::ListConnectedPins() const
+{
+	TArray<UBlueprintGraphPin*> Result;
+	UEdGraphPin* Pin = GetNativePin();
+	if(!Pin)
+	{
+		return Result;
+	}
+
+	for(UEdGraphPin* LinkedToPin : Pin->LinkedTo)
+	{
+		UBlueprintGraphPin* WrappedPin = UBlueprintGraphPin::FromNativePin(LinkedToPin);
+		if(WrappedPin)
+		{
+			Result.Add(WrappedPin);
+		}
+	}
+	return Result;
+}
+
 UEdGraphPin* UBlueprintGraphPin::GetNativePin() const
 {
 	if(!::IsValid(Node))
diff --git a/Engine/Source/Editor/BlueprintEditorLibrary/Private/BlueprintEditorLibrary/BlueprintGraphPin.h b/Engine/Source/Editor/BlueprintEditorLibrary/Private/BlueprintEditorLibrary/BlueprintGraphPin.h
--- a/Engine/Source/Editor/BlueprintEditorLibrary/Private/BlueprintEditorLibrary/BlueprintGraphPin.h
+++ b/Engine/Source/Editor/BlueprintEditorLibrary/Private/BlueprintEditorLibrary/BlueprintGraphPin.h
@@ -16,7 +16,7 @@ class UK2Node;
  * Helper for interfacing with UEdGraphPin (not a UObject, relatively low level) as it's used in blueprints. This
  * is a transient class and as such is created on demand and does not have useful identity, so comparison of reference
  * equality to other UBlueprintGraphPins will not be meaningful. Wrapping UEdGraphPin makes this type less prone to
- * leaking and safer to use across graph edit operations.
+ * leaking and safer to use across graph edit operations. To establish pin equivalence use IsSameNativePin()
  */
 UCLASS(Transient, BlueprintType, HideDropDown)
 class UBlueprintGraphPin : public UObject
@@ -25,6 +25,16 @@ class UBlueprintGraphPin : public UObject
 public:
 	/** Creates a UBlueprintGraphPin wrappin the underlying native Pin */
 	static UBlueprintGraphPin* FromNativePin(UEdGraphPin* Pin);
+	
+	/** 
+ 	 * Because these pin objects are proxies for the underlying native pin, this function
+	 * should be used for establishing whether to UBlueprintGraphPin objects refer to the 
+	 * same pin on a node
+	 * 
+	 * @return true if the Other pin and this Pin are the same pin
+	 */
+	UFUNCTION(BlueprintCallable, Category = "Blueprint Editor")
+	bool IsSameNativePin(const UBlueprintGraphPin* Other) const;
 
 	/**
 	 * Sets the literal value associated with this pin, returning false if the value is not valid
@@ -74,6 +84,12 @@ public:
 	UFUNCTION(BlueprintCallable, Category = "Blueprint Editor")
 	FString GetPinTypeAsJsonSchema() const;
 
+	/**
+	 * Returns the string used to describe the pin type in the blueprint UI
+	 */
+	UFUNCTION(BlueprintCallable, Category = "Blueprint Editor")
+	FText GetPinTypeDisplayString() const;
+
 	/**
 	 * Returns the pin's type, FEdGraphPinType is mostly opaque to script but can
 	 * be used e.g. to create variables of the same type
@@ -94,6 +110,13 @@ public:
 	UFUNCTION(BlueprintCallable, Category = "Blueprint Editor")
 	UK2Node* GetOwningNode() const;
 
+	/**
+	 * Returns a list of pins this pin is connected to, note that these pin objects are 
+	 * proxies and so comparison to other pin objects must be done via IsSameNativePin
+	 */
+	UFUNCTION(BlueprintCallable, Category = "Blueprint Editor")
+	TArray<UBlueprintGraphPin*> ListConnectedPins() const;
+
 	UEdGraphPin* GetNativePin() const;
 	static UEdGraphPin* GetNativePinSafe(const UBlueprintGraphPin* Pin);
 private:
diff --git a/Engine/Source/Editor/BlueprintEditorLibrary/Public/BlueprintEditorLibrary.h b/Engine/Source/Editor/BlueprintEditorLibrary/Public/BlueprintEditorLibrary.h
--- a/Engine/Source/Editor/BlueprintEditorLibrary/Public/BlueprintEditorLibrary.h
+++ b/Engine/Source/Editor/BlueprintEditorLibrary/Public/BlueprintEditorLibrary.h
@@ -448,7 +448,12 @@ public:
 	 */
 	UFUNCTION(BlueprintCallable, Category = "Blueprint Graph Editing", meta = (ScriptMethod))
 	static FIntPoint GetNodePos(const UK2Node* Node);
-
+	
+	/**
+	 * Returns this node's title, as used when the node is displayed in a list
+	 */
+	UFUNCTION(BlueprintCallable, Category = "Blueprint Graph Editing", meta = (ScriptMethod))
+	static FString GetNodeTitle(const UK2Node* Node);
 	/// end UK2Node extension/script methods
 
 	/// begin UK2Node_IfThenElse textension/script methods:
