Added functions to seach and spawn nodes by categories and names.

#rnx

#Synced-CL 51104521

diff --git a/Engine/Source/Editor/BlueprintEditorLibrary/Private/BlueprintEditorLibrary.cpp b/Engine/Source/Editor/BlueprintEditorLibrary/Private/BlueprintEditorLibrary.cpp
--- a/Engine/Source/Editor/BlueprintEditorLibrary/Private/BlueprintEditorLibrary.cpp
+++ b/Engine/Source/Editor/BlueprintEditorLibrary/Private/BlueprintEditorLibrary.cpp
@@ -1696,6 +1696,16 @@ FString UBlueprintEditorLibrary::GetNodeTitle(const UK2Node* Node)
 	return Node->GetNodeTitle(ENodeTitleType::ListView).ToString();
 }
 
+FString UBlueprintEditorLibrary::GetNodeCategory(const UK2Node* Node)
+{
+	if (!::IsValid(Node))
+	{
+		return TEXT("");
+	}
+
+	return Node->GetMenuCategory().ToString();
+}
+
 UBlueprintGraphPin* UBlueprintEditorLibrary::FindElsePin(const UK2Node_IfThenElse* Node)
 {
 	if(!::IsValid(Node))
diff --git a/Engine/Source/Editor/BlueprintEditorLibrary/Private/BlueprintEditorLibrary/BlueprintGraphEditor.cpp b/Engine/Source/Editor/BlueprintEditorLibrary/Private/BlueprintEditorLibrary/BlueprintGraphEditor.cpp
--- a/Engine/Source/Editor/BlueprintEditorLibrary/Private/BlueprintEditorLibrary/BlueprintGraphEditor.cpp
+++ b/Engine/Source/Editor/BlueprintEditorLibrary/Private/BlueprintEditorLibrary/BlueprintGraphEditor.cpp
@@ -27,9 +27,11 @@
 #include "UObject/SoftObjectPtr.h"
 #include "UObject/SoftObjectPath.h"
 #include "BlueprintEditorLibrary/BlueprintGraphPin.h"
+#include "BlueprintActionDatabase.h"
 
 #define LOCTEXT_NAMESPACE "UBlueprintGraphEditor"
 
+
 UBlueprintGraphEditor::UBlueprintGraphEditor()
 {
 }
@@ -1090,4 +1092,142 @@ void UBlueprintGraphEditor::SetFunctionFlagImpl(UFunction* Function, UK2Node_Fun
 	}
 }
 
+void UBlueprintGraphEditor::ListAvailableNodes(TArray<FString>& Result, const TArray<UBlueprintGraphPin*>& ContextPins) const
+{
+	TRACE_CPUPROFILER_EVENT_SCOPE(UBlueprintGraphEditor::ListAvailableNodes);
+
+	Result.Reset();
+	if (!IsValid())
+	{
+		return;
+	}
+	FBlueprintActionFilter Filter;
+	FBlueprintActionContext& FilterContext = Filter.Context;
+	for (UBlueprintGraphPin* ContextPin : ContextPins) 
+	{
+		if (UEdGraphPin* NativePin = UBlueprintGraphPin::GetNativePinSafe(ContextPin))
+		{
+			FilterContext.Pins.Add(NativePin);
+		}
+	}
+
+	FilterContext.Graphs.Add(Graph);
+	FilterContext.Blueprints.Add(Blueprint);
+
+	FBlueprintActionDatabase& ActionDatabase = FBlueprintActionDatabase::Get();
+	FBlueprintActionDatabase::FActionRegistry const& ActionRegistry = ActionDatabase.GetAllActions();
+
+	for (auto Iterator(ActionRegistry.CreateConstIterator()); Iterator; ++Iterator)
+	{
+		const FObjectKey& ObjKey = Iterator->Key;
+		const FBlueprintActionDatabase::FActionList& ActionList = Iterator->Value;
+
+		if (UObject* ActionObject = ObjKey.ResolveObjectPtr())
+		{
+			for (UBlueprintNodeSpawner const* NodeSpawner : ActionList)
+			{
+				FBlueprintActionInfo BlueprintAction(ActionObject, NodeSpawner);
+
+				bool bPassedFilter = !Filter.IsFiltered(BlueprintAction);
+
+				if (bPassedFilter)
+				{
+					UBlueprintNodeSpawner const* Action = BlueprintAction.NodeSpawner;
+					if (!Action)
+					{
+						continue;
+					}
+					FBlueprintActionUiSpec UISpec = Action->GetUiSpec(FilterContext, BlueprintAction.GetBindings());
+					FString NameAndCategory = FString::Printf(TEXT("%s|%s"),
+						*UISpec.Category.ToString().Replace(TEXT(" "), TEXT("")),
+						*UISpec.MenuName.ToString().Replace(TEXT(" "), TEXT("")));
+					Result.Add(NameAndCategory);	
+				}
+			}
+		}
+		else
+		{
+			// Remove this (invalid) entry on the next tick.
+			ActionDatabase.DeferredRemoveEntry(ObjKey);
+		}
+	}
+}
+
+
+UEdGraphNode* UBlueprintGraphEditor::CreateNodeFromName(const FString& NodeWithCategory, FVector2D const Location, const TArray<UBlueprintGraphPin*>& ContextPins)
+{
+	TRACE_CPUPROFILER_EVENT_SCOPE(UBlueprintGraphEditor::ListAvailableNodes);
+
+	if (!IsValid())
+	{
+		return nullptr;
+	}
+	FBlueprintActionFilter Filter;
+	FBlueprintActionContext& FilterContext = Filter.Context;
+	FilterContext.Graphs.Add(Graph);
+	FilterContext.Blueprints.Add(Blueprint);
+	for (UBlueprintGraphPin* ContextPin : ContextPins)
+	{
+		if (UEdGraphPin* NativePin = UBlueprintGraphPin::GetNativePinSafe(ContextPin))
+		{
+			FilterContext.Pins.Add(NativePin);
+		}
+	}
+
+	FBlueprintActionDatabase& ActionDatabase = FBlueprintActionDatabase::Get();
+	FBlueprintActionDatabase::FActionRegistry const& ActionRegistry = ActionDatabase.GetAllActions();
+
+	for (auto Iterator(ActionRegistry.CreateConstIterator()); Iterator; ++Iterator)
+	{
+		const FObjectKey& ObjKey = Iterator->Key;
+		const FBlueprintActionDatabase::FActionList& ActionList = Iterator->Value;
+
+		if (UObject* ActionObject = ObjKey.ResolveObjectPtr())
+		{
+			for (UBlueprintNodeSpawner const* NodeSpawner : ActionList)
+			{
+				FBlueprintActionInfo BlueprintAction(ActionObject, NodeSpawner);
+
+				bool bPassedFilter = !Filter.IsFiltered(BlueprintAction);
+
+				if (bPassedFilter)
+				{
+					UBlueprintNodeSpawner const* Action = BlueprintAction.NodeSpawner;
+					if (!Action) 
+					{
+						continue;
+					}
+					FBlueprintActionUiSpec UISpec = Action->GetUiSpec(FilterContext, BlueprintAction.GetBindings());
+					FString NameAndCategory = FString::Printf(TEXT("%s|%s"),
+						*UISpec.Category.ToString().Replace(TEXT(" "), TEXT("")),
+						*UISpec.MenuName.ToString().Replace(TEXT(" "), TEXT("")));
+					if (NodeWithCategory.Equals(NameAndCategory, ESearchCase::CaseSensitive))
+					{
+						UEdGraphNode* SpawnedNode = Action->Invoke(Graph, BlueprintAction.GetBindings(), Location);
+						if (!SpawnedNode)
+						{
+							continue;
+						}
+						for (UEdGraphPin* Pin : FilterContext.Pins)
+						{
+							SpawnedNode->AutowireNewNode(Pin);
+						}
+						if (UK2Node* K2Node = Cast<UK2Node>(SpawnedNode)) 
+						{
+							HandleNodeSpawned(K2Node);
+						}
+						return SpawnedNode;
+					}
+				}
+			}
+		}
+		else
+		{
+			// Remove this (invalid) entry on the next tick.
+			ActionDatabase.DeferredRemoveEntry(ObjKey);
+		}
+	}
+	return nullptr;
+}
+
 #undef LOCTEXT_NAMESPACE
diff --git a/Engine/Source/Editor/BlueprintEditorLibrary/Private/BlueprintEditorLibrary/BlueprintGraphEditor.h b/Engine/Source/Editor/BlueprintEditorLibrary/Private/BlueprintEditorLibrary/BlueprintGraphEditor.h
--- a/Engine/Source/Editor/BlueprintEditorLibrary/Private/BlueprintEditorLibrary/BlueprintGraphEditor.h
+++ b/Engine/Source/Editor/BlueprintEditorLibrary/Private/BlueprintEditorLibrary/BlueprintGraphEditor.h
@@ -69,6 +69,9 @@ public:
 	UFUNCTION(BlueprintCallable, Category = "Blueprint Editor")
 	static UBlueprintGraphEditor* GetGraphEditor(UEdGraph* InGraph);
 
+	UFUNCTION(BlueprintCallable, Category = "Blueprint Editor")
+	UEdGraph* GetGraph() const { return Graph; }
+
 	/**
 	 * Lists all nodes in the graph
 	 */
@@ -95,6 +98,22 @@ public:
 	UFUNCTION(BlueprintCallable, Category = "Blueprint Editor")
 	void ListNodesWithNotes(TArray<UK2Node*>& OutNodes);
 	
+	/**
+	 * List all of the nodes by name that can be added to this graph
+	 */
+	UFUNCTION(BlueprintCallable, Category = "Blueprint Editor", meta=(AutoCreateRefTerm = "ContextPins"))
+	void ListAvailableNodes(TArray<FString>& Result, const TArray<UBlueprintGraphPin*>& ContextPins) const;
+
+	/**
+	 * Creates node from Catagory and Name. Example:
+	 *			- 'Development|PrintString'
+	 *			- 'Utilities|Operators|Add'
+	 *			- 'Utilities|FlowControl|Branch'
+	 *			- 'Math|Vector|VectorLength'
+	 */
+	UFUNCTION(BlueprintCallable, Category = "Blueprint Editor", meta = (AutoCreateRefTerm = "ContextPins"))
+	UEdGraphNode* CreateNodeFromName(const FString& NodeWithCategory, FVector2D const Location, const TArray<UBlueprintGraphPin*>& ContextPins);
+
 	/**
 	 * Adds a graph input parameter with name InputName, and optional PinType and Value. Returns the added pin.
 	 * By default the type will be a boolean as defined by UBlueprintEditorLibrary.
diff --git a/Engine/Source/Editor/BlueprintEditorLibrary/Public/BlueprintEditorLibrary.h b/Engine/Source/Editor/BlueprintEditorLibrary/Public/BlueprintEditorLibrary.h
--- a/Engine/Source/Editor/BlueprintEditorLibrary/Public/BlueprintEditorLibrary.h
+++ b/Engine/Source/Editor/BlueprintEditorLibrary/Public/BlueprintEditorLibrary.h
@@ -454,6 +454,13 @@ public:
 	 */
 	UFUNCTION(BlueprintCallable, Category = "Blueprint Graph Editing", meta = (ScriptMethod))
 	static FString GetNodeTitle(const UK2Node* Node);
+
+	/**
+	 * Returns this node's menu category
+	 */
+	UFUNCTION(BlueprintCallable, Category = "Blueprint Graph Editing", meta = (ScriptMethod))
+	static FString GetNodeCategory(const UK2Node* Node);
+
 	/// end UK2Node extension/script methods
 
 	/// begin UK2Node_IfThenElse textension/script methods:
