Fix mistake in implementation of AddGetMemberVariableNode/AddSetMemberVariableNode

#rnx

#Synced-CL 51251105

diff --git a/Engine/Source/Editor/BlueprintEditorLibrary/Private/BlueprintEditorLibrary/BlueprintGraphEditor.cpp b/Engine/Source/Editor/BlueprintEditorLibrary/Private/BlueprintEditorLibrary/BlueprintGraphEditor.cpp
--- a/Engine/Source/Editor/BlueprintEditorLibrary/Private/BlueprintEditorLibrary/BlueprintGraphEditor.cpp
+++ b/Engine/Source/Editor/BlueprintEditorLibrary/Private/BlueprintEditorLibrary/BlueprintGraphEditor.cpp
@@ -741,7 +741,7 @@ UK2Node_VariableGet* UBlueprintGraphEditor::AddGetMemberVariableNode(FName Membe
 		return nullptr;
 	}
 	
-	UBlueprintVariableNodeSpawner* Spawner = UBlueprintVariableNodeSpawner::CreateFromMemberOrParam(UK2Node_VariableGet::StaticClass(), Property, Graph);
+	UBlueprintVariableNodeSpawner* Spawner = UBlueprintVariableNodeSpawner::CreateFromMemberOrParam(UK2Node_VariableGet::StaticClass(), Property);
 	UK2Node_VariableGet* Result = CastChecked<UK2Node_VariableGet>(Spawner->Invoke(Graph, {}, GetNodeLocation()));
 	if(!Result)
 	{
@@ -778,7 +778,7 @@ UK2Node_VariableSet* UBlueprintGraphEditor::AddSetMemberVariableNode(FName Membe
 		return nullptr;
 	}
 	
-	UBlueprintVariableNodeSpawner* Spawner = UBlueprintVariableNodeSpawner::CreateFromMemberOrParam(UK2Node_VariableSet::StaticClass(), Property, Graph);
+	UBlueprintVariableNodeSpawner* Spawner = UBlueprintVariableNodeSpawner::CreateFromMemberOrParam(UK2Node_VariableSet::StaticClass(), Property);
 	UK2Node_VariableSet* Result = CastChecked<UK2Node_VariableSet>(Spawner->Invoke(Graph, {}, GetNodeLocation()));
 	if(!Result)
 	{
