Add blueprint editor library functions for:
* listing member variables
* Getting member variable types
* listing local variables
* getting local variable types
* setting the default value on local variables, as these cannot be set via set_editor_property

#rnx
#Synced-CL 51516599

diff --git a/Engine/Source/Editor/BlueprintEditorLibrary/Private/BlueprintEditorLibrary.cpp b/Engine/Source/Editor/BlueprintEditorLibrary/Private/BlueprintEditorLibrary.cpp
--- a/Engine/Source/Editor/BlueprintEditorLibrary/Private/BlueprintEditorLibrary.cpp
+++ b/Engine/Source/Editor/BlueprintEditorLibrary/Private/BlueprintEditorLibrary.cpp
@@ -169,8 +169,83 @@ namespace InternalBlueprintEditorLibrary
 
 		return Result;
 	}
+
+	// Iterator helper for walking a class hierarchy's member variables, whether they are native or not - return false to break the loop:
+	void ForEachInheritedVariable(const UBlueprint* BP, TFunctionRef<bool(const UClass*, const FProperty*, const FBPVariableDescription*)> Ftor);
+	bool GetVariableType(const FProperty* P, const FBPVariableDescription* D, FEdGraphPinType& OutType);
+	FString GetVariableFullName(const UClass* C, const FProperty* P, const FBPVariableDescription* D);
+	FString GetVariableShortName(const FProperty* P, const FBPVariableDescription* D);
 };
 
+void InternalBlueprintEditorLibrary::ForEachInheritedVariable(const UBlueprint* BP, TFunctionRef<bool(const UClass*, const FProperty*, const FBPVariableDescription*)> Ftor)
+{
+	check(BP);
+	UClass* Class = BP->ParentClass;
+	while(Class)
+	{
+		const FString ClassPath = Class->GetPathName();
+		if (const UBlueprint* ParentBP = Cast<UBlueprint>(Class->ClassGeneratedBy))
+		{
+			for (const FBPVariableDescription& VarDesc : ParentBP->NewVariables)
+			{
+				if(!Ftor(Class, nullptr, &VarDesc))
+				{
+					return;
+				}
+			}
+			Class = ParentBP->ParentClass;
+		}
+		else
+		{
+			for (TFieldIterator<FProperty> PropertyIt(Class, EFieldIteratorFlags::ExcludeSuper); PropertyIt; ++PropertyIt)
+			{
+				if(!Ftor(Class, *PropertyIt, nullptr))
+				{
+					return;
+				}
+			}
+			Class = Class->GetSuperClass();
+		}
+	}
+}
+
+bool InternalBlueprintEditorLibrary::GetVariableType(const FProperty* P, const FBPVariableDescription* D, FEdGraphPinType& OutPinType)
+{
+	check(P || D);
+	if(D)
+	{
+		OutPinType = D->VarType;
+		return true;
+	}
+
+	FEdGraphPinType PinType;
+	const UEdGraphSchema_K2* Schema = GetDefault<UEdGraphSchema_K2>();
+	return Schema->ConvertPropertyToPinType(P, OutPinType);
+}
+
+FString InternalBlueprintEditorLibrary::GetVariableFullName(const UClass* C, const FProperty* P, const FBPVariableDescription* D)
+{
+	check(C);
+	check(P || D);
+	const FString ClassPath = C->GetPathName();
+	if(D)
+	{
+		return FString::Printf(TEXT("%s.%s"), *ClassPath, *D->VarName.ToString());
+	}
+
+	return FString::Printf(TEXT("%s.%s"), *ClassPath, *P->GetName());
+}
+
+FString InternalBlueprintEditorLibrary::GetVariableShortName(const FProperty* P, const FBPVariableDescription* D)
+{
+	check(P || D);
+	if(D)
+	{
+		return D->VarName.ToString();
+	}
+	return P->GetName();
+}
+
 ///////////////////////////////////////////////////////////
 // UBlueprintEditorLibrary
 
@@ -1026,6 +1101,109 @@ void UBlueprintEditorLibrary::ListGraphNames(const UBlueprint* Blueprint, TArray
 	}
 }
 
+void UBlueprintEditorLibrary::ListMemberVariableNames(const UBlueprint* Blueprint, TArray<FString>& OutMemberVariableNames, bool bIncludeInheritedMembers)
+{
+	if (!IsValid(Blueprint))
+	{
+		UE_LOGFMT(LogBlueprintEditorLib, Warning, "Invalid blueprint has no member variables");
+		return;
+	}
+
+	// Add variables declared directly on this blueprint
+	for (const FBPVariableDescription& VarDesc : Blueprint->NewVariables)
+	{
+		OutMemberVariableNames.Add(VarDesc.VarName.ToString());
+	}
+
+	if (bIncludeInheritedMembers)
+	{
+		InternalBlueprintEditorLibrary::ForEachInheritedVariable(Blueprint, [&OutMemberVariableNames](const UClass* Class, const FProperty* P, const FBPVariableDescription* D)
+			{
+				OutMemberVariableNames.Add(InternalBlueprintEditorLibrary::GetVariableFullName(Class, P, D));
+				return true; // we want all variables
+			});
+	}
+}
+
+TOptional<FEdGraphPinType> UBlueprintEditorLibrary::GetMemberVariableType(const UBlueprint* Blueprint, const FString& MemberVariableName)
+{
+	if (!IsValid(Blueprint))
+	{
+		UE_LOGFMT(LogBlueprintEditorLib, Warning, "Invalid blueprint has no member variables");
+		return {};
+	}
+	
+	// first, search for a member declared on this blueprint:
+	for (const FBPVariableDescription& VarDesc : Blueprint->NewVariables)
+	{
+		if(VarDesc.VarName.ToString().Equals(MemberVariableName))
+		{
+			return VarDesc.VarType;
+		}
+	}
+	
+	// no match within the class, search the hierarchy - looking for a fully qualified match:
+	TOptional<FEdGraphPinType> Result;
+	const UEdGraphSchema_K2* Schema = GetDefault<UEdGraphSchema_K2>();
+	InternalBlueprintEditorLibrary::ForEachInheritedVariable(Blueprint, [&Result, &MemberVariableName, Schema](const UClass* Class, const FProperty* P, const FBPVariableDescription* D)
+		{
+			const FString MemberPath = InternalBlueprintEditorLibrary::GetVariableFullName(Class, P, D);
+			if(MemberPath.Equals(MemberVariableName))
+			{
+				FEdGraphPinType PinType;
+				if(InternalBlueprintEditorLibrary::GetVariableType(P, D, PinType))
+				{
+					Result = PinType;
+					return false; // done, matched
+				}
+			}
+			return true; // keep searching
+		});
+	if(Result.IsSet())
+	{
+		return Result;
+	}
+	
+	for (const FBPVariableDescription& VarDesc : Blueprint->NewVariables)
+	{
+		if(VarDesc.VarName.ToString().Equals(MemberVariableName, ESearchCase::IgnoreCase))
+		{
+			return VarDesc.VarType;
+		}
+	}
+	if(Result.IsSet())
+	{
+		UE_LOGFMT(LogBlueprintEditorLib, Warning, "Case ignored to match member name - use a fully qualifed name to find inherited values unambiguously: {MemberName}", MemberVariableName);
+		return Result;
+	}
+
+	// Still no match, do a short match in the hierarchy - this could result in ambiguous searches 
+	// so leave a warning if we've matched:
+	InternalBlueprintEditorLibrary::ForEachInheritedVariable(Blueprint, [&Result, &MemberVariableName, Schema](const UClass* Class, const FProperty* P, const FBPVariableDescription* D)
+		{
+			FString MemberName = InternalBlueprintEditorLibrary::GetVariableShortName(P, D);
+			if(MemberName.Equals( MemberVariableName, ESearchCase::IgnoreCase))
+			{
+				FEdGraphPinType PinType;
+				if(InternalBlueprintEditorLibrary::GetVariableType(P, D, PinType))
+				{
+					Result = PinType;
+					return false; // done, matched
+				}
+			}
+			return true; // keep searching
+		});
+	if(Result.IsSet())
+	{
+		UE_LOGFMT(LogBlueprintEditorLib, Warning, "Short name matched inherited variable - consider fully qualifying the name: {MemberVariableName}", MemberVariableName);
+		return Result;
+	}
+
+	// property not matched at all, log a warning:
+	UE_LOGFMT(LogBlueprintEditorLib, Warning, "Could not find a match for: {MemberVariableName}", MemberVariableName);
+	return {};
+}
+
 void UBlueprintEditorLibrary::SetBlueprintVariableExposeToCinematics(UBlueprint* Blueprint, const FName& VariableName, bool bExposeToCinematics)
 {
 	if (!Blueprint)
diff --git a/Engine/Source/Editor/BlueprintEditorLibrary/Private/BlueprintEditorLibrary/BlueprintGraphEditor.cpp b/Engine/Source/Editor/BlueprintEditorLibrary/Private/BlueprintEditorLibrary/BlueprintGraphEditor.cpp
--- a/Engine/Source/Editor/BlueprintEditorLibrary/Private/BlueprintEditorLibrary/BlueprintGraphEditor.cpp
+++ b/Engine/Source/Editor/BlueprintEditorLibrary/Private/BlueprintEditorLibrary/BlueprintGraphEditor.cpp
@@ -349,6 +349,129 @@ bool UBlueprintGraphEditor::RemoveGraphOutputParameter(FName OutputName)
 	return true;
 }
 
+void UBlueprintGraphEditor::ListLocalVariableNames(TArray<FString>& OutLocalVariableNames) const
+{
+	if(!IsValid())
+	{
+		return;
+	}
+
+	if(!FBlueprintEditorUtils::DoesSupportLocalVariables(Graph))
+	{
+		UE_LOGFMT(LogBlueprint, Error, "Graph does not support locals {GraphName}", Graph->GetPathName());
+		return;
+	}
+
+	TArray<UK2Node_FunctionEntry*> EntryNodes;
+	Graph->GetNodesOfClass(EntryNodes);
+
+	// There should always be an entry node in the function graph
+	if(EntryNodes.Num() == 0)
+	{
+		UE_LOGFMT(LogBlueprint, Error, "Function graph is malformed - has no entry nodes {GraphName}", Graph->GetPathName());
+		return;
+	}
+				
+	UK2Node_FunctionEntry* FuncEntry = EntryNodes[0];
+	for (const FBPVariableDescription& LocalVar : FuncEntry->LocalVariables)
+	{
+		OutLocalVariableNames.Add(LocalVar.VarName.ToString());
+	}
+}
+
+TOptional<FEdGraphPinType> UBlueprintGraphEditor::GetLocalVariableType(const FString& LocalVariableName) const
+{
+	if(!IsValid())
+	{
+		return {};
+	}
+
+	if(!FBlueprintEditorUtils::DoesSupportLocalVariables(Graph))
+	{
+		UE_LOGFMT(LogBlueprint, Error, "Graph does not support locals {GraphName}", Graph->GetPathName());
+		return {};
+	}
+	
+	const FBPVariableDescription* VariableDescription = FBlueprintEditorUtils::FindLocalVariable(Blueprint, Graph, FName(*LocalVariableName));
+	if(!VariableDescription)
+	{
+		UE_LOGFMT(LogBlueprint, Warning, "No local {LocalName} found in {GraphName}", LocalVariableName, Graph->GetPathName());
+		return {};
+	}
+
+	return VariableDescription->VarType;
+}
+
+TOptional<FString> UBlueprintGraphEditor::GetLocalVariableDefaultValue(const FString& LocalVariableName) const
+{
+	if(!IsValid())
+	{
+		return {};
+	}
+
+	if(!FBlueprintEditorUtils::DoesSupportLocalVariables(Graph))
+	{
+		UE_LOGFMT(LogBlueprint, Error, "Graph does not support locals {GraphName}", Graph->GetPathName());
+		return {};
+	}
+	
+	const FBPVariableDescription* VariableDescription = FBlueprintEditorUtils::FindLocalVariable(Blueprint, Graph, FName(*LocalVariableName));
+	if(!VariableDescription)
+	{
+		UE_LOGFMT(LogBlueprint, Warning, "No local {LocalName} found in {GraphName}", LocalVariableName, Graph->GetPathName());
+		return {};
+	}
+
+	return VariableDescription->DefaultValue;
+}
+
+bool UBlueprintGraphEditor::SetLocalVariableDefaultValue(const FString& LocalVariableName, const FString& NewDefaultValue)
+{
+	if(!IsValid())
+	{
+		return false;
+	}
+
+	if(!FBlueprintEditorUtils::DoesSupportLocalVariables(Graph))
+	{
+		UE_LOGFMT(LogBlueprint, Error, "Graph does not support locals {GraphName}", Graph->GetPathName());
+		return false;
+	}
+	
+	FName PropertyAsName(*LocalVariableName);
+	FBPVariableDescription* VariableDescription = FBlueprintEditorUtils::FindLocalVariable(Blueprint, Graph, PropertyAsName);
+	if(!VariableDescription)
+	{
+		UE_LOGFMT(LogBlueprint, Warning, "No local {LocalName} found in {GraphName}", LocalVariableName, Graph->GetPathName());
+		return false;
+	}
+
+	// we have to get the underlying property, which may require a skeleton compile:
+	const FProperty* VariableProperty = FindLocalByName(PropertyAsName);
+	if(!VariableProperty)
+	{
+		UE_LOGFMT(LogBlueprint, Error, "No backign property found for {LocalName} found in {GraphName} - possibly missing skeleton compile", LocalVariableName, Graph->GetPathName());
+		return false;
+	}
+
+	const UEdGraphSchema_K2* Schema = GetDefault<UEdGraphSchema_K2>();
+	void* PropertyOnScope =  VariableProperty->AllocateAndInitializeValue();
+	ON_SCOPE_EXIT
+	{
+		VariableProperty->DestroyAndFreeValue(PropertyOnScope);
+	};
+	constexpr EPropertyPortFlags ImportTextPortFlags = EPropertyPortFlags::PPF_SerializedAsImportText;
+	const TCHAR* Result = VariableProperty->ImportText_Direct(*NewDefaultValue, PropertyOnScope, nullptr, ImportTextPortFlags, GLog);
+	if(Result == nullptr)
+	{
+		return false;
+	}
+
+	// the value imported correctly, lets store it:
+	VariableDescription->DefaultValue = NewDefaultValue;
+	return true;
+}
+
 bool UBlueprintGraphEditor::AddLocalVariable(FName LocalName, const FEdGraphPinType& PinType, const FString& Value)
 {
 	if(!IsValid())
@@ -1046,6 +1169,32 @@ const FProperty* UBlueprintGraphEditor::FindParameterByName(FName Name) const
 	return nullptr;
 }
 
+const FProperty* UBlueprintGraphEditor::FindLocalByName(FName Name) const
+{
+	const UFunction* SkeletonFunction = GetSkeletonFunction();
+	if(!SkeletonFunction)
+	{
+		return nullptr;
+	}
+
+	TFieldIterator<FProperty> ParamIt(SkeletonFunction);
+	// skip params:
+	while(ParamIt && (ParamIt->PropertyFlags & CPF_Parm))
+	{
+		++ParamIt;
+	}
+	while(ParamIt)
+	{
+		if((*ParamIt)->GetFName() == Name)
+		{
+			return *ParamIt;
+		}
+		++ParamIt;
+	}
+
+	return nullptr;
+}
+
 const UFunction* UBlueprintGraphEditor::GetSkeletonFunction() const
 {
 	return Blueprint->SkeletonGeneratedClass ? 
diff --git a/Engine/Source/Editor/BlueprintEditorLibrary/Private/BlueprintEditorLibrary/BlueprintGraphEditor.h b/Engine/Source/Editor/BlueprintEditorLibrary/Private/BlueprintEditorLibrary/BlueprintGraphEditor.h
--- a/Engine/Source/Editor/BlueprintEditorLibrary/Private/BlueprintEditorLibrary/BlueprintGraphEditor.h
+++ b/Engine/Source/Editor/BlueprintEditorLibrary/Private/BlueprintEditorLibrary/BlueprintGraphEditor.h
@@ -155,6 +155,31 @@ public:
 	UFUNCTION(BlueprintCallable, Category = "Blueprint Editor")
 	bool RemoveGraphOutputParameter(FName OutputName);
 	
+	/**
+	 * Lists the names of the local variables in the function.
+	 */
+	UFUNCTION(BlueprintCallable, Category = "Blueprint Editor")
+	void ListLocalVariableNames(TArray<FString>& OutLocalVariableNames) const;
+
+	/**
+	 * Returns the type of a local variable, unset if no variable is found.
+	 */
+	UFUNCTION(BlueprintCallable, Category = "Blueprint Editor")
+	TOptional<FEdGraphPinType> GetLocalVariableType(const FString& LocalVariableName) const;
+	
+	/**
+	 * Returns the default value of a local variable, unset if no variable is found.
+	 */
+	UFUNCTION(BlueprintCallable, Category = "Blueprint Editor")
+	TOptional<FString> GetLocalVariableDefaultValue(const FString& LocalVariableName) const;
+
+	/**
+	 * Set default value of a local variable. Returns false if the operation fails
+	 * either because the value is invalid or there is no variable of the specified name.
+	 */
+	UFUNCTION(BlueprintCallable, Category = "Blueprint Editor")
+	bool SetLocalVariableDefaultValue(const FString& LocalVariableName, const FString& NewDefaultValue);
+
 	/**
 	 * Adds a local variable, use BlueprintEditorLibrary to specify a type. A default value can optionally be provided.
 	 * By default the type will be a boolean as defined by UBlueprintEditorLibrary.
@@ -299,6 +324,7 @@ private:
 	UK2Node_FunctionEntry* GetFunctionEntryNode() const;
 	void HandleNodeSpawned(UK2Node* Result);
 	const FProperty* FindParameterByName(FName Name) const;
+	const FProperty* FindLocalByName(FName Name) const;
 	const UFunction* GetSkeletonFunction() const;
 	UFunction* GetSkeletonFunction();
 	static const UClass* GetClassFromPath(const FString& ClassPath);
diff --git a/Engine/Source/Editor/BlueprintEditorLibrary/Public/BlueprintEditorLibrary.h b/Engine/Source/Editor/BlueprintEditorLibrary/Public/BlueprintEditorLibrary.h
--- a/Engine/Source/Editor/BlueprintEditorLibrary/Public/BlueprintEditorLibrary.h
+++ b/Engine/Source/Editor/BlueprintEditorLibrary/Public/BlueprintEditorLibrary.h
@@ -323,6 +323,18 @@ public:
 	UFUNCTION(BlueprintCallable, Category = "Blueprint Editor")
 	static void ListGraphNames(const UBlueprint* Blueprint, TArray<FName>& OutGraphNames);
 
+	/**
+	 * Lists the names of the member variables in the BP. Inherited members will be prefixed with the full path of their declaring class.
+	 */
+	UFUNCTION(BlueprintCallable, Category = "Blueprint Editor", meta = (ScriptMethod))
+	static void ListMemberVariableNames(const UBlueprint* Blueprint, TArray<FString>& OutMemberVariableNames, bool bIncludeInheritedMembers = true);
+	
+	/**
+	 * Returns the type of a member variable, unset if no variable is found. Inherited variables can be requested using a fully qualified name to avoid ambiguities.
+	 */
+	UFUNCTION(BlueprintCallable, Category = "Blueprint Editor", meta = (ScriptMethod))
+	static TOptional<FEdGraphPinType> GetMemberVariableType(const UBlueprint* Blueprint, const FString& MemberVariableName);
+
 	/**
 	  * Adds a member variable to the specified blueprint inferring the type from a provided value.
 	  * 
