Changes for Blueprint Toolset for EDA
Added bool to blueprint compile
Made ErrorMsg on EdGraphNode Blueprint readable
Made Status on Blueprints Blueprint readable

#Synced-CL 51605467

diff --git a/Engine/Source/Editor/BlueprintEditorLibrary/Private/BlueprintEditorLibrary.cpp b/Engine/Source/Editor/BlueprintEditorLibrary/Private/BlueprintEditorLibrary.cpp
--- a/Engine/Source/Editor/BlueprintEditorLibrary/Private/BlueprintEditorLibrary.cpp
+++ b/Engine/Source/Editor/BlueprintEditorLibrary/Private/BlueprintEditorLibrary.cpp
@@ -584,14 +584,16 @@ void UBlueprintEditorLibrary::UpgradeOperatorNodes(UBlueprint* Blueprint)
 	}
 }
 
-void UBlueprintEditorLibrary::CompileBlueprint(UBlueprint* Blueprint)
+bool UBlueprintEditorLibrary::CompileBlueprint(UBlueprint* Blueprint)
 {
-	if (Blueprint)
+	if (!Blueprint)
 	{
-		// Skip saving this to avoid possible tautologies when saving and allow the user to manually save
-		EBlueprintCompileOptions Flags = EBlueprintCompileOptions::SkipSave;
-		FKismetEditorUtilities::CompileBlueprint(Blueprint, Flags);
+		return false;
 	}
+	// Skip saving this to avoid possible tautologies when saving and allow the user to manually save
+	EBlueprintCompileOptions Flags = EBlueprintCompileOptions::SkipSave;
+	FKismetEditorUtilities::CompileBlueprint(Blueprint, Flags);
+	return Blueprint->Status != BS_Error;
 }
 
 UEdGraph* UBlueprintEditorLibrary::AddFunctionGraph(UBlueprint* Blueprint, const FString& FuncName)
diff --git a/Engine/Source/Editor/BlueprintEditorLibrary/Public/BlueprintEditorLibrary.h b/Engine/Source/Editor/BlueprintEditorLibrary/Public/BlueprintEditorLibrary.h
--- a/Engine/Source/Editor/BlueprintEditorLibrary/Public/BlueprintEditorLibrary.h
+++ b/Engine/Source/Editor/BlueprintEditorLibrary/Public/BlueprintEditorLibrary.h
@@ -153,9 +153,11 @@ public:
 	* Compiles the given blueprint. 
 	*
 	* @param Blueprint	Blueprint to compile
+	* 
+	* @return True if compilation succeeded with no errors
 	*/
 	UFUNCTION(BlueprintCallable, Category = "Blueprint Upgrade Tools")
-	static UE_API void CompileBlueprint(UBlueprint* Blueprint);
+	static UE_API bool CompileBlueprint(UBlueprint* Blueprint);
 
 	/**
 	* Adds a function to the given blueprint
diff --git a/Engine/Source/Runtime/Engine/Classes/EdGraph/EdGraphNode.h b/Engine/Source/Runtime/Engine/Classes/EdGraph/EdGraphNode.h
--- a/Engine/Source/Runtime/Engine/Classes/EdGraph/EdGraphNode.h
+++ b/Engine/Source/Runtime/Engine/Classes/EdGraph/EdGraphNode.h
@@ -398,7 +398,7 @@ public:
 	int32 ErrorType;
 	
 	/** Error/Warning description */
-	UPROPERTY()
+	UPROPERTY(BlueprintReadOnly, Category="EdGraphNode")
 	FString ErrorMsg;
 	
 	/** GUID to uniquely identify this node, to facilitate diffing versions of this graph */
diff --git a/Engine/Source/Runtime/Engine/Classes/Engine/Blueprint.h b/Engine/Source/Runtime/Engine/Classes/Engine/Blueprint.h
--- a/Engine/Source/Runtime/Engine/Classes/Engine/Blueprint.h
+++ b/Engine/Source/Runtime/Engine/Classes/Engine/Blueprint.h
@@ -504,7 +504,7 @@ public:
 	EBlueprintCompileMode CompileMode;
 
 	/** The current status of this blueprint */
-	UPROPERTY(transient)
+	UPROPERTY(transient, Category=BlueprintOptions, BlueprintReadOnly)
 	TEnumAsByte<enum EBlueprintStatus> Status;
 
 	/** Overrides the BP's display name in the editor UI */
