Make blueprint graph pin a value type so that comparison behaves as expected

#rnx
#Synced-CL 51516599

diff --git a/Engine/Source/Editor/BlueprintEditorLibrary/Private/BlueprintEditorLibrary.cpp b/Engine/Source/Editor/BlueprintEditorLibrary/Private/BlueprintEditorLibrary.cpp
--- a/Engine/Source/Editor/BlueprintEditorLibrary/Private/BlueprintEditorLibrary.cpp
+++ b/Engine/Source/Editor/BlueprintEditorLibrary/Private/BlueprintEditorLibrary.cpp
@@ -1528,20 +1528,20 @@ static UEdGraphPin* GetIfTypeMatches(const UClass* Class, UEdGraphPin* Pin, cons
 }
 }
 
-TArray<UBlueprintGraphPin*> UBlueprintEditorLibrary::ListAllPins(const UK2Node* Node, TEnumAsByte<EEdGraphPinDirection> InDirection)
+TArray<FBlueprintGraphPin> UBlueprintEditorLibrary::ListAllPins(const UK2Node* Node, TEnumAsByte<EEdGraphPinDirection> InDirection)
 {
 	if(!::IsValid(Node))
 	{
 		return {};
 	}
 
-	TArray<UBlueprintGraphPin*> Result;
+	TArray<FBlueprintGraphPin> Result;
 	// first input pins - stable order will be more useful to callers
 	if(InDirection == EGPD_MAX || InDirection == EGPD_Input)
 	{
 		UE::Private::ForEachVisibleInputPin(Node, [&Result](UEdGraphPin* Pin)
 		{
-			Result.Add(UBlueprintGraphPin::FromNativePin(Pin));
+			Result.Add(UBlueprintGraphPinLibrary::FromNativePin(Pin));
 			return true;
 		});
 	}
@@ -1550,23 +1550,33 @@ TArray<UBlueprintGraphPin*> UBlueprintEditorLibrary::ListAllPins(const UK2Node*
 	{
 		UE::Private::ForEachVisibleOutputPin(Node, [&Result](UEdGraphPin* Pin)
 		{
-			Result.Add(UBlueprintGraphPin::FromNativePin(Pin));
+			Result.Add(UBlueprintGraphPinLibrary::FromNativePin(Pin));
 			return true;
 		});
 	}
 	return Result;
 }
 
-UBlueprintGraphPin* UBlueprintEditorLibrary::FindInputPinByIndex(const UK2Node* Node, int32 Index, const FEdGraphPinType& Type)
+TArray<FBlueprintGraphPin> UBlueprintEditorLibrary::ListInputPins(const UK2Node* Node)
+{
+	return ListAllPins(Node, EGPD_Input);
+}
+
+TArray<FBlueprintGraphPin> UBlueprintEditorLibrary::ListOutputPins(const UK2Node* Node)
+{
+	return ListAllPins(Node, EGPD_Output);
+}
+
+FBlueprintGraphPin UBlueprintEditorLibrary::FindInputPinByIndex(const UK2Node* Node, int32 Index, const FEdGraphPinType& Type)
 {
 	if(!::IsValid(Node))
 	{
-		return nullptr;
+		return FBlueprintGraphPin{};
 	}
 
 	const UEdGraphSchema_K2* K2Schema = GetDefault<UEdGraphSchema_K2>();
-	UBlueprintGraphPin* Result = nullptr;
-	int32 Iter = 0; 
+	FBlueprintGraphPin Result;
+	int32 Iter = 0;
 	const UClass* Class = Node->GetBlueprintClassFromNode();
 	UE::Private::ForEachVisibleInputPin(Node, [&Iter, Class, Node, Index, &Type, K2Schema, &Result](UEdGraphPin* Pin) {
 		if(Iter != Index)
@@ -1575,219 +1585,219 @@ UBlueprintGraphPin* UBlueprintEditorLibrary::FindInputPinByIndex(const UK2Node*
 			return true; // wrong index, keep searching
 		}
 
-		Result = UBlueprintGraphPin::FromNativePin(UE::Private::GetIfTypeMatches(Class, Pin, Type));
-		if(!Result)
+		Result = UBlueprintGraphPinLibrary::FromNativePin(UE::Private::GetIfTypeMatches(Class, Pin, Type));
+		if(!::IsValid(Result.Node))
 		{
-			UE_LOGFMT(LogBlueprint, Warning, 
+			UE_LOGFMT(LogBlueprint, Warning,
 				"Type on input pin {pin_index} named {pin_name} on {node} did not match {request_type}",
-				Index, 
-				Pin->PinName, 
-				Node->GetPathName(), 
+				Index,
+				Pin->PinName,
+				Node->GetPathName(),
 				PinTypeToJsonSchema(Type, Class) );
 		}
 		return false; // index matched, all done
 	});
-	
-	if(!Result)
+
+	if(!::IsValid(Result.Node))
 	{
 		UE_LOGFMT(LogBlueprint, Warning, "Missing input pin at: {pin_index}, Pin count: {iter}, on node {node}", Index, Iter, Node->GetPathName());
 	}
 	return Result;
 }
 
-UBlueprintGraphPin* UBlueprintEditorLibrary::FindInputPin(const UK2Node* Node, FName PinName, const FEdGraphPinType& Type)
+FBlueprintGraphPin UBlueprintEditorLibrary::FindInputPin(const UK2Node* Node, FName PinName, const FEdGraphPinType& Type)
 {
 	if(!::IsValid(Node))
 	{
-		return nullptr;
+		return FBlueprintGraphPin{};
 	}
 
-	UBlueprintGraphPin* Result = nullptr;
+	FBlueprintGraphPin Result;
 	const UClass* Class = Node->GetBlueprintClassFromNode();
-	UE::Private::ForEachVisibleInputPin(Node, [PinName, Class, Node, &Type, &Result](UEdGraphPin* Pin) 
+	UE::Private::ForEachVisibleInputPin(Node, [PinName, Class, Node, &Type, &Result](UEdGraphPin* Pin)
 	{
 		if(Pin->PinName != PinName)
 		{
 			return true;
 		}
-		
-		Result = UBlueprintGraphPin::FromNativePin(UE::Private::GetIfTypeMatches(Class, Pin, Type));
-		if(!Result)
+
+		Result = UBlueprintGraphPinLibrary::FromNativePin(UE::Private::GetIfTypeMatches(Class, Pin, Type));
+		if(!::IsValid(Result.Node))
 		{
-			UE_LOGFMT(LogBlueprint, Warning, 
+			UE_LOGFMT(LogBlueprint, Warning,
 				"Type of input pin {pin_name} on {node} did not match {request_type}",
-				Pin->PinName, 
-				Node->GetPathName(), 
+				Pin->PinName,
+				Node->GetPathName(),
 				PinTypeToJsonSchema(Type, Class) );
 		}
 		return false;
 	});
-	
-	if(!Result)
+
+	if(!::IsValid(Result.Node))
 	{
 		UE_LOGFMT(LogBlueprint, Warning, "Missing input pin named: {pin_name}, on node {node}", PinName, Node->GetPathName());
 	}
 	return Result;
 }
 
-UBlueprintGraphPin* UBlueprintEditorLibrary::FindOutputPinByIndex(const UK2Node* Node, int32 Index, const FEdGraphPinType& Type)
+FBlueprintGraphPin UBlueprintEditorLibrary::FindOutputPinByIndex(const UK2Node* Node, int32 Index, const FEdGraphPinType& Type)
 {
 	if(!::IsValid(Node))
 	{
-		return nullptr;
+		return FBlueprintGraphPin{};
 	}
 
 	const UEdGraphSchema_K2* K2Schema = GetDefault<UEdGraphSchema_K2>();
-	UBlueprintGraphPin* Result = nullptr;
-	int32 Iter = 0; 
+	FBlueprintGraphPin Result;
+	int32 Iter = 0;
 	const UClass* Class = Node->GetBlueprintClassFromNode();
-	UE::Private::ForEachVisibleOutputPin(Node, [&Iter, Class, Node, Index, &Type, K2Schema, &Result](UEdGraphPin* Pin) 
+	UE::Private::ForEachVisibleOutputPin(Node, [&Iter, Class, Node, Index, &Type, K2Schema, &Result](UEdGraphPin* Pin)
 	{
 		if(Iter != Index)
 		{
 			++Iter;
 			return true; // wrong index, keep searching
 		}
-		
-		Result = UBlueprintGraphPin::FromNativePin(UE::Private::GetIfTypeMatches(Class, Pin, Type));
-		if(!Result)
+
+		Result = UBlueprintGraphPinLibrary::FromNativePin(UE::Private::GetIfTypeMatches(Class, Pin, Type));
+		if(!::IsValid(Result.Node))
 		{
-			UE_LOGFMT(LogBlueprint, Warning, 
+			UE_LOGFMT(LogBlueprint, Warning,
 				"Type on output pin {pin_index} named {pin_name} on {node} did not match {request_type}",
-				Index, 
-				Pin->PinName, 
-				Node->GetPathName(), 
+				Index,
+				Pin->PinName,
+				Node->GetPathName(),
 				PinTypeToJsonSchema(Type, Class) );
 		}
 		return false; // index matched, all done
 	});
-	
-	if(!Result)
+
+	if(!::IsValid(Result.Node))
 	{
 		UE_LOGFMT(LogBlueprint, Warning, "Missing output pin at: {pin_index}, Pin count: {iter}, on node {node}", Index, Iter, Node->GetPathName());
 	}
 	return Result;
 }
 
-UBlueprintGraphPin* UBlueprintEditorLibrary::FindOutputPin(const UK2Node* Node, FName PinName, const FEdGraphPinType& Type)
+FBlueprintGraphPin UBlueprintEditorLibrary::FindOutputPin(const UK2Node* Node, FName PinName, const FEdGraphPinType& Type)
 {
 	if(!::IsValid(Node))
 	{
-		return nullptr;
+		return FBlueprintGraphPin{};
 	}
 
-	UBlueprintGraphPin* Result = nullptr;
+	FBlueprintGraphPin Result;
 	const UClass* Class = Node->GetBlueprintClassFromNode();
-	UE::Private::ForEachVisibleOutputPin(Node, [Class, Node, PinName, &Type, &Result](UEdGraphPin* Pin) 
+	UE::Private::ForEachVisibleOutputPin(Node, [Class, Node, PinName, &Type, &Result](UEdGraphPin* Pin)
 	{
 		if(Pin->PinName != PinName)
 		{
 			return true;
 		}
-		
-		Result = UBlueprintGraphPin::FromNativePin(UE::Private::GetIfTypeMatches(Class, Pin, Type));
-		if(!Result)
+
+		Result = UBlueprintGraphPinLibrary::FromNativePin(UE::Private::GetIfTypeMatches(Class, Pin, Type));
+		if(!::IsValid(Result.Node))
 		{
-			UE_LOGFMT(LogBlueprint, Warning, 
+			UE_LOGFMT(LogBlueprint, Warning,
 				"Type of output pin {pin_name} on {node} did not match {request_type}",
-				Pin->PinName, 
-				Node->GetPathName(), 
+				Pin->PinName,
+				Node->GetPathName(),
 				PinTypeToJsonSchema(Type, Class) );
 		}
 		return false;
 	});
-	
-	if(!Result)
+
+	if(!::IsValid(Result.Node))
 	{
 		UE_LOGFMT(LogBlueprint, Warning, "Missing output pin named: {pin_name}, on node {node}", PinName, Node->GetPathName());
 	}
 	return Result;
 }
 
-UBlueprintGraphPin* UBlueprintEditorLibrary::FindExecutePin(const UK2Node* Node)
+FBlueprintGraphPin UBlueprintEditorLibrary::FindExecutePin(const UK2Node* Node)
 {
 	if(!::IsValid(Node))
 	{
-		return nullptr;
+		return FBlueprintGraphPin{};
 	}
 
-	UBlueprintGraphPin* Result = nullptr;
-	UE::Private::ForEachVisibleInputPin(Node, [&Result](UEdGraphPin* Pin) 
+	FBlueprintGraphPin Result;
+	UE::Private::ForEachVisibleInputPin(Node, [&Result](UEdGraphPin* Pin)
 	{
 		if(UEdGraphSchema_K2::IsExecPin(*Pin) && Pin->PinName == UEdGraphSchema_K2::PN_Execute)
 		{
-			Result = UBlueprintGraphPin::FromNativePin(Pin);
+			Result = UBlueprintGraphPinLibrary::FromNativePin(Pin);
 			return false;
 		}
 		return true;
 	});
-	if(!Result)
+	if(!::IsValid(Result.Node))
 	{
 		UE_LOGFMT(LogBlueprint, Warning, "No execute pin found on node {node}", Node->GetPathName());
 	}
 	return Result;
 }
 
-UBlueprintGraphPin* UBlueprintEditorLibrary::FindThenPin(const UK2Node* Node)
+FBlueprintGraphPin UBlueprintEditorLibrary::FindThenPin(const UK2Node* Node)
 {
 	if(!::IsValid(Node))
 	{
-		return nullptr;
+		return FBlueprintGraphPin{};
 	}
 
-	UBlueprintGraphPin* Result = nullptr;
-	UE::Private::ForEachVisibleOutputPin(Node, [&Result](UEdGraphPin* Pin) 
+	FBlueprintGraphPin Result;
+	UE::Private::ForEachVisibleOutputPin(Node, [&Result](UEdGraphPin* Pin)
 	{
 		if(UEdGraphSchema_K2::IsExecPin(*Pin) && Pin->PinName == UEdGraphSchema_K2::PN_Then)
 		{
-			Result = UBlueprintGraphPin::FromNativePin(Pin);
+			Result = UBlueprintGraphPinLibrary::FromNativePin(Pin);
 			return false;
 		}
 		return true;
 	});
-	if(!Result)
+	if(!::IsValid(Result.Node))
 	{
 		UE_LOGFMT(LogBlueprint, Warning, "No then pin found on node {node}", Node->GetPathName());
 	}
 	return Result;
 }
 
-UBlueprintGraphPin* UBlueprintEditorLibrary::FindSelfPin(const UK2Node* Node)
+FBlueprintGraphPin UBlueprintEditorLibrary::FindSelfPin(const UK2Node* Node)
 {
 	if(!::IsValid(Node))
 	{
-		return nullptr;
+		return FBlueprintGraphPin{};
 	}
 
-	UBlueprintGraphPin* Result = nullptr;
+	FBlueprintGraphPin Result;
 	const UEdGraphSchema_K2* K2Schema = GetDefault<UEdGraphSchema_K2>();
-	UE::Private::ForEachVisibleInputPin(Node, [&Result, K2Schema](UEdGraphPin* Pin) 
+	UE::Private::ForEachVisibleInputPin(Node, [&Result, K2Schema](UEdGraphPin* Pin)
 	{
 		if(K2Schema->IsSelfPin(*Pin))
 		{
-			Result = UBlueprintGraphPin::FromNativePin(Pin);
+			Result = UBlueprintGraphPinLibrary::FromNativePin(Pin);
 			return false;
 		}
 		return true;
 	});
-	if(!Result)
+	if(!::IsValid(Result.Node))
 	{
 		UE_LOGFMT(LogBlueprint, Warning, "No self pin found on node {node}", Node->GetPathName());
 	}
 	return Result;
 }
 
-UBlueprintGraphPin* UBlueprintEditorLibrary::FindResultPin(const UK2Node* Node)
+FBlueprintGraphPin UBlueprintEditorLibrary::FindResultPin(const UK2Node* Node)
 {
 	if(!::IsValid(Node))
 	{
-		return nullptr;
+		return FBlueprintGraphPin{};
 	}
 
 	UEdGraphPin* Result = nullptr;
 	const UEdGraphSchema_K2* K2Schema = GetDefault<UEdGraphSchema_K2>();
 	bool bDidWarn = false;
-	UE::Private::ForEachVisibleOutputPin(Node, [&Result, Node, &bDidWarn, K2Schema](UEdGraphPin* Pin) 
+	UE::Private::ForEachVisibleOutputPin(Node, [&Result, Node, &bDidWarn, K2Schema](UEdGraphPin* Pin)
 	{
 		if(UEdGraphSchema_K2::IsExecPin(*Pin))
 		{
@@ -1795,7 +1805,7 @@ UBlueprintGraphPin* UBlueprintEditorLibrary::FindResultPin(const UK2Node* Node)
 		}
 		if(Result)
 		{
-			Result = nullptr; // multiple outputs found, return nullptr
+			Result = nullptr; // multiple outputs found, return empty
 			UE_LOGFMT(LogBlueprint, Warning, "Multiple result pins found on node {node}", Node->GetPathName());
 			bDidWarn = true;
 			return false;
@@ -1807,20 +1817,20 @@ UBlueprintGraphPin* UBlueprintEditorLibrary::FindResultPin(const UK2Node* Node)
 	{
 		UE_LOGFMT(LogBlueprint, Warning, "No result pin found on node {node}", Node->GetPathName());
 	}
-	return UBlueprintGraphPin::FromNativePin(Result);
+	return UBlueprintGraphPinLibrary::FromNativePin(Result);
 }
 
-UBlueprintGraphPin* UBlueprintEditorLibrary::FindDataInputPin(const UK2Node* Node)
+FBlueprintGraphPin UBlueprintEditorLibrary::FindDataInputPin(const UK2Node* Node)
 {
 	if(!::IsValid(Node))
 	{
-		return nullptr;
+		return FBlueprintGraphPin{};
 	}
 
 	UEdGraphPin* Result = nullptr;
 	const UEdGraphSchema_K2* K2Schema = GetDefault<UEdGraphSchema_K2>();
 	bool bDidWarn = false;
-	UE::Private::ForEachVisibleInputPin(Node, [&Result, &bDidWarn, Node, K2Schema](UEdGraphPin* Pin) 
+	UE::Private::ForEachVisibleInputPin(Node, [&Result, &bDidWarn, Node, K2Schema](UEdGraphPin* Pin)
 	{
 		if(UEdGraphSchema_K2::IsExecPin(*Pin))
 		{
@@ -1828,7 +1838,7 @@ UBlueprintGraphPin* UBlueprintEditorLibrary::FindDataInputPin(const UK2Node* Nod
 		}
 		if(Result)
 		{
-			Result = nullptr; // multiple inputs found, return nullptr
+			Result = nullptr; // multiple inputs found, return empty
 			UE_LOGFMT(LogBlueprint, Warning, "Multiple input data pins found on node {node}", Node->GetPathName());
 			bDidWarn = true;
 			return false;
@@ -1840,7 +1850,7 @@ UBlueprintGraphPin* UBlueprintEditorLibrary::FindDataInputPin(const UK2Node* Nod
 	{
 		UE_LOGFMT(LogBlueprint, Warning, "No data input pin found on node {node}", Node->GetPathName());
 	}
-	return UBlueprintGraphPin::FromNativePin(Result);
+	return UBlueprintGraphPinLibrary::FromNativePin(Result);
 }
 
 void UBlueprintEditorLibrary::SetNodePos(UK2Node* Node, FIntPoint Pos)
@@ -1886,24 +1896,24 @@ FString UBlueprintEditorLibrary::GetNodeCategory(const UK2Node* Node)
 	return Node->GetMenuCategory().ToString();
 }
 
-UBlueprintGraphPin* UBlueprintEditorLibrary::FindElsePin(const UK2Node_IfThenElse* Node)
+FBlueprintGraphPin UBlueprintEditorLibrary::FindElsePin(const UK2Node_IfThenElse* Node)
 {
 	if(!::IsValid(Node))
 	{
-		return nullptr;
+		return FBlueprintGraphPin{};
 	}
 
-	return UBlueprintGraphPin::FromNativePin(Node->GetElsePin());
+	return UBlueprintGraphPinLibrary::FromNativePin(Node->GetElsePin());
 }
-	
-UBlueprintGraphPin* UBlueprintEditorLibrary::FindConditionPin(const UK2Node_IfThenElse* Node)
+
+FBlueprintGraphPin UBlueprintEditorLibrary::FindConditionPin(const UK2Node_IfThenElse* Node)
 {
 	if(!::IsValid(Node))
 	{
-		return nullptr;
+		return FBlueprintGraphPin{};
 	}
 
-	return UBlueprintGraphPin::FromNativePin(Node->GetConditionPin());
+	return UBlueprintGraphPinLibrary::FromNativePin(Node->GetConditionPin());
 }
 
 #undef LOCTEXT_NAMESPACE	// "BlueprintEditorLibrary"
diff --git a/Engine/Source/Editor/BlueprintEditorLibrary/Private/BlueprintEditorLibrary/BlueprintGraphEditor.cpp b/Engine/Source/Editor/BlueprintEditorLibrary/Private/BlueprintEditorLibrary/BlueprintGraphEditor.cpp
--- a/Engine/Source/Editor/BlueprintEditorLibrary/Private/BlueprintEditorLibrary/BlueprintGraphEditor.cpp
+++ b/Engine/Source/Editor/BlueprintEditorLibrary/Private/BlueprintEditorLibrary/BlueprintGraphEditor.cpp
@@ -216,14 +216,14 @@ bool UBlueprintGraphEditor::RemoveGraphInputParameter(FName InputName)
 	return bHasPin;
 }
 
-UBlueprintGraphPin* UBlueprintGraphEditor::FindGraphEntryPin() const
+FBlueprintGraphPin UBlueprintGraphEditor::FindGraphEntryPin() const
 {
 	UK2Node_FunctionEntry* Entry = GetFunctionEntryNode();
 	if(!Entry)
 	{
-		return nullptr;
+		return FBlueprintGraphPin{};
 	}
-	return UBlueprintEditorLibrary::FindThenPin(Entry);//Entry->FindThenPin();
+	return UBlueprintEditorLibrary::FindThenPin(Entry);
 }
 
 UK2Node_Event* UBlueprintGraphEditor::FindEventNode(FName EventName) const
@@ -1065,31 +1065,31 @@ void UBlueprintGraphEditor::RemoveNodes(const TArray<UK2Node*>& Nodes)
 	}
 }
 
-UBlueprintGraphPin* UBlueprintGraphEditor::AddGraphInputParameter(FName InputName, const FEdGraphPinType& PinType, const FString& Value)
+FBlueprintGraphPin UBlueprintGraphEditor::AddGraphInputParameter(FName InputName, const FEdGraphPinType& PinType, const FString& Value)
 {
 	if(!IsValid())
 	{
-		return nullptr;
+		return FBlueprintGraphPin{};
 	}
 
 	UK2Node_FunctionEntry* Entry = GetFunctionEntryNode();
 	if(!Entry)
 	{
-		return nullptr;
+		return FBlueprintGraphPin{};
 	}
 
 	const UEdGraphSchema_K2* Schema = GetDefault<UEdGraphSchema_K2>();
 	Entry->Modify();
-	UEdGraphPin* Pin = Entry->CreateUserDefinedPin(InputName, 
-		PinType != FEdGraphPinType() ? PinType : UBlueprintEditorLibrary::GetBasicTypeByName(UEdGraphSchema_K2::PC_Boolean), 
+	UEdGraphPin* Pin = Entry->CreateUserDefinedPin(InputName,
+		PinType != FEdGraphPinType() ? PinType : UBlueprintEditorLibrary::GetBasicTypeByName(UEdGraphSchema_K2::PC_Boolean),
 		EGPD_Output);
 	if(Pin)
 	{
 		Schema->TrySetDefaultValue(*Pin, Value, false);
-		return UBlueprintGraphPin::FromNativePin(Pin);
+		return UBlueprintGraphPinLibrary::FromNativePin(Pin);
 	}
 
-	return nullptr;
+	return FBlueprintGraphPin{};
 }
 
 bool UBlueprintGraphEditor::IsValid() const
@@ -1241,7 +1241,7 @@ void UBlueprintGraphEditor::SetFunctionFlagImpl(UFunction* Function, UK2Node_Fun
 	}
 }
 
-void UBlueprintGraphEditor::ListAvailableNodes(TArray<FString>& Result, const TArray<UBlueprintGraphPin*>& ContextPins) const
+void UBlueprintGraphEditor::ListAvailableNodes(TArray<FString>& Result, const TArray<FBlueprintGraphPin>& ContextPins) const
 {
 	TRACE_CPUPROFILER_EVENT_SCOPE(UBlueprintGraphEditor::ListAvailableNodes);
 
@@ -1252,9 +1252,9 @@ void UBlueprintGraphEditor::ListAvailableNodes(TArray<FString>& Result, const TA
 	}
 	FBlueprintActionFilter Filter;
 	FBlueprintActionContext& FilterContext = Filter.Context;
-	for (UBlueprintGraphPin* ContextPin : ContextPins) 
+	for (const FBlueprintGraphPin& ContextPin : ContextPins)
 	{
-		if (UEdGraphPin* NativePin = UBlueprintGraphPin::GetNativePinSafe(ContextPin))
+		if (UEdGraphPin* NativePin = UBlueprintGraphPinLibrary::GetNativePinSafe(ContextPin))
 		{
 			FilterContext.Pins.Add(NativePin);
 		}
@@ -1303,7 +1303,7 @@ void UBlueprintGraphEditor::ListAvailableNodes(TArray<FString>& Result, const TA
 }
 
 
-UEdGraphNode* UBlueprintGraphEditor::CreateNodeFromName(const FString& NodeWithCategory, FVector2D const Location, const TArray<UBlueprintGraphPin*>& ContextPins)
+UEdGraphNode* UBlueprintGraphEditor::CreateNodeFromName(const FString& NodeWithCategory, FVector2D const Location, const TArray<FBlueprintGraphPin>& ContextPins)
 {
 	TRACE_CPUPROFILER_EVENT_SCOPE(UBlueprintGraphEditor::ListAvailableNodes);
 
@@ -1315,9 +1315,9 @@ UEdGraphNode* UBlueprintGraphEditor::CreateNodeFromName(const FString& NodeWithC
 	FBlueprintActionContext& FilterContext = Filter.Context;
 	FilterContext.Graphs.Add(Graph);
 	FilterContext.Blueprints.Add(Blueprint);
-	for (UBlueprintGraphPin* ContextPin : ContextPins)
+	for (const FBlueprintGraphPin& ContextPin : ContextPins)
 	{
-		if (UEdGraphPin* NativePin = UBlueprintGraphPin::GetNativePinSafe(ContextPin))
+		if (UEdGraphPin* NativePin = UBlueprintGraphPinLibrary::GetNativePinSafe(ContextPin))
 		{
 			FilterContext.Pins.Add(NativePin);
 		}
diff --git a/Engine/Source/Editor/BlueprintEditorLibrary/Private/BlueprintEditorLibrary/BlueprintGraphEditor.h b/Engine/Source/Editor/BlueprintEditorLibrary/Private/BlueprintEditorLibrary/BlueprintGraphEditor.h
--- a/Engine/Source/Editor/BlueprintEditorLibrary/Private/BlueprintEditorLibrary/BlueprintGraphEditor.h
+++ b/Engine/Source/Editor/BlueprintEditorLibrary/Private/BlueprintEditorLibrary/BlueprintGraphEditor.h
@@ -6,6 +6,7 @@
 #include "EdGraph/EdGraphNode.h"
 #include "EdGraph/EdGraphPin.h"
 #include "Templates/SubclassOf.h"
+#include "BlueprintEditorLibrary/BlueprintGraphPin.h"
 
 #include "BlueprintGraphEditor.generated.h"
 
@@ -21,7 +22,6 @@ class UK2Node_Select;
 class UK2Node_FunctionResult;
 class UK2Node_VariableGet;
 class UK2Node_VariableSet;
-class UBlueprintGraphPin;
 
 /**
  * Object oriented helper for editing graphs in a blueprint object
@@ -102,7 +102,7 @@ public:
 	 * List all of the nodes by name that can be added to this graph
 	 */
 	UFUNCTION(BlueprintCallable, Category = "Blueprint Editor", meta=(AutoCreateRefTerm = "ContextPins"))
-	void ListAvailableNodes(TArray<FString>& Result, const TArray<UBlueprintGraphPin*>& ContextPins) const;
+	void ListAvailableNodes(TArray<FString>& Result, const TArray<FBlueprintGraphPin>& ContextPins) const;
 
 	/**
 	 * Creates node from Catagory and Name. Example:
@@ -112,14 +112,14 @@ public:
 	 *			- 'Math|Vector|VectorLength'
 	 */
 	UFUNCTION(BlueprintCallable, Category = "Blueprint Editor", meta = (AutoCreateRefTerm = "ContextPins"))
-	UEdGraphNode* CreateNodeFromName(const FString& NodeWithCategory, FVector2D const Location, const TArray<UBlueprintGraphPin*>& ContextPins);
+	UEdGraphNode* CreateNodeFromName(const FString& NodeWithCategory, FVector2D const Location, const TArray<FBlueprintGraphPin>& ContextPins);
 
 	/**
 	 * Adds a graph input parameter with name InputName, and optional PinType and Value. Returns the added pin.
 	 * By default the type will be a boolean as defined by UBlueprintEditorLibrary.
 	 */
 	UFUNCTION(BlueprintCallable, Category = "Blueprint Editor")
-	UBlueprintGraphPin* AddGraphInputParameter(FName InputName, const FEdGraphPinType& PinType = FEdGraphPinType(), const FString& Value = FString());
+	FBlueprintGraphPin AddGraphInputParameter(FName InputName, const FEdGraphPinType& PinType = FEdGraphPinType(), const FString& Value = FString());
 	
 	/**
 	 * Removes the graph input parameter specified by InputName
@@ -131,7 +131,7 @@ public:
 	 * Returns the 'entry' pin for the graph, e.g. the 'then' pin of the entry node
 	 */
 	UFUNCTION(BlueprintCallable, Category = "Blueprint Editor")
-	UBlueprintGraphPin* FindGraphEntryPin() const;
+	FBlueprintGraphPin FindGraphEntryPin() const;
 
 	/**
 	 * Returns the event node associated with EventName - aka Member Name
diff --git a/Engine/Source/Editor/BlueprintEditorLibrary/Private/BlueprintEditorLibrary/BlueprintGraphPin.cpp b/Engine/Source/Editor/BlueprintEditorLibrary/Private/BlueprintEditorLibrary/BlueprintGraphPin.cpp
--- a/Engine/Source/Editor/BlueprintEditorLibrary/Private/BlueprintEditorLibrary/BlueprintGraphPin.cpp
+++ b/Engine/Source/Editor/BlueprintEditorLibrary/Private/BlueprintEditorLibrary/BlueprintGraphPin.cpp
@@ -7,49 +7,81 @@
 #include "K2Node.h"
 #include "Logging/StructuredLog.h"
 
-UBlueprintGraphPin* UBlueprintGraphPin::FromNativePin(UEdGraphPin* Pin)
+bool FBlueprintGraphPin::operator==(const FBlueprintGraphPin&) const = default;
+bool FBlueprintGraphPin::operator!=(const FBlueprintGraphPin&) const = default;
+
+FBlueprintGraphPin UBlueprintGraphPinLibrary::FromNativePin(UEdGraphPin* Pin)
 {
 	if(!Pin)
 	{
-		return nullptr;
+		return FBlueprintGraphPin{};
 	}
 
 	UK2Node* OwningNode = Cast<UK2Node>(Pin->GetOwningNode());
 	if(!OwningNode)
 	{
-		return nullptr;
+		return FBlueprintGraphPin{};
+	}
+
+	FBlueprintGraphPin ScriptPin;
+	ScriptPin.Node = OwningNode;
+	ScriptPin.PinDirection = Pin->Direction;
+	ScriptPin.PinName = Pin->PinName;
+
+	int32 LogicalIndex = 0;
+	for(UEdGraphPin* NodePin : OwningNode->Pins)
+	{
+		if(NodePin->bHidden)
+		{
+			continue;
+		}
+		if(NodePin->Direction == Pin->Direction)
+		{
+			if(NodePin == Pin)
+			{
+				ScriptPin.PinIndex = LogicalIndex;
+				break;
+			}
+			LogicalIndex++;
+		}
 	}
 
-	UBlueprintGraphPin* ScriptPin = NewObject<UBlueprintGraphPin>();
-	ScriptPin->Node = OwningNode;
-	ScriptPin->PinDirection = Pin->Direction;
-	ScriptPin->PinName = Pin->PinName;
 	return ScriptPin;
 }
 
-bool UBlueprintGraphPin::IsSameNativePin(const UBlueprintGraphPin* Other) const
+bool UBlueprintGraphPinLibrary::IsValid(const FBlueprintGraphPin& Pin)
+{
+	if( Pin == FBlueprintGraphPin() )
+	{
+		return false;
+	}
+
+	return true;
+}
+
+bool UBlueprintGraphPinLibrary::IsSameNativePin(const FBlueprintGraphPin& Pin, const FBlueprintGraphPin& Other)
 {
+	const UEdGraphPin* NativePin = GetNativePinSafe(Pin);
 	const UEdGraphPin* OtherNativePin = GetNativePinSafe(Other);
-	const UEdGraphPin* NativePin = GetNativePin();
-	if(!OtherNativePin || !NativePin)
+	if(!NativePin || !OtherNativePin)
 	{
 		return false;
 	}
 
-	return OtherNativePin == NativePin;
+	return NativePin == OtherNativePin;
 }
 
-bool UBlueprintGraphPin::SetPinValue(const FString& Value)
-{						 
+bool UBlueprintGraphPinLibrary::SetPinValue(const FBlueprintGraphPin& Pin, const FString& Value)
+{
 	const UEdGraphSchema_K2* Schema = GetDefault<UEdGraphSchema_K2>();
-	if(UEdGraphPin* Pin = GetNativePin())
+	if(UEdGraphPin* NativePin = GetNativePinSafe(Pin))
 	{
-		Schema->TrySetDefaultValue(*Pin, Value, false);
+		Schema->TrySetDefaultValue(*NativePin, Value, false);
 		// Reset default values
-		if (!Schema->IsPinDefaultValid(Pin, Pin->DefaultValue, Pin->DefaultObject, Pin->DefaultTextValue).IsEmpty())
+		if (!Schema->IsPinDefaultValid(NativePin, NativePin->DefaultValue, NativePin->DefaultObject, NativePin->DefaultTextValue).IsEmpty())
 		{
-			Schema->ResetPinToAutogeneratedDefaultValue(Pin, false);
-			UE_LOGFMT(LogBlueprint, Warning, "Failed to set default value {Value} on {PinName}", Value, Pin->PinName);
+			Schema->ResetPinToAutogeneratedDefaultValue(NativePin, false);
+			UE_LOGFMT(LogBlueprint, Warning, "Failed to set default value {Value} on {PinName}", Value, NativePin->PinName);
 			return false;
 		}
 		return true;
@@ -57,15 +89,10 @@ bool UBlueprintGraphPin::SetPinValue(const FString& Value)
 	return false;
 }
 
-bool UBlueprintGraphPin::TryCreateConnection(UBlueprintGraphPin* Other)
+bool UBlueprintGraphPinLibrary::TryCreateConnection(const FBlueprintGraphPin& Pin, const FBlueprintGraphPin& Other)
 {
-	if(!Other)
-	{
-		return false;
-	}
-
-	UEdGraphPin* SelfNativePin = GetNativePin();
-	UEdGraphPin* OtherNativePin = Other->GetNativePin();
+	UEdGraphPin* SelfNativePin = GetNativePinSafe(Pin);
+	UEdGraphPin* OtherNativePin = GetNativePinSafe(Other);
 	if(!SelfNativePin || !OtherNativePin)
 	{
 		return false;
@@ -75,41 +102,36 @@ bool UBlueprintGraphPin::TryCreateConnection(UBlueprintGraphPin* Other)
 	return Schema->TryCreateConnection(SelfNativePin, OtherNativePin);
 }
 
-bool UBlueprintGraphPin::CanCreateConnection(const UBlueprintGraphPin* Other) const
+bool UBlueprintGraphPinLibrary::CanCreateConnection(const FBlueprintGraphPin& Pin, const FBlueprintGraphPin& Other)
 {
+	const UEdGraphPin* NativePin = GetNativePinSafe(Pin);
 	const UEdGraphPin* OtherNativePin = GetNativePinSafe(Other);
-	const UEdGraphPin* NativePin = GetNativePin();
-	if(!OtherNativePin || !NativePin)
+	if(!NativePin || !OtherNativePin)
 	{
 		return false;
 	}
-	
+
 	const UEdGraphSchema_K2* Schema = GetDefault<UEdGraphSchema_K2>();
-	return Schema->CanCreateConnection(OtherNativePin, NativePin).Response != CONNECT_RESPONSE_DISALLOW;
+	return Schema->CanCreateConnection(NativePin, OtherNativePin).Response != CONNECT_RESPONSE_DISALLOW;
 }
-	
-bool UBlueprintGraphPin::BreakPinLinks()
+
+bool UBlueprintGraphPinLibrary::BreakPinLinks(const FBlueprintGraphPin& Pin)
 {
-	UEdGraphPin* SelfNativePin = GetNativePin();
-	if(!SelfNativePin)
+	UEdGraphPin* NativePin = GetNativePinSafe(Pin);
+	if(!NativePin)
 	{
 		return false;
 	}
-	
+
 	const UEdGraphSchema_K2* Schema = GetDefault<UEdGraphSchema_K2>();
-	Schema->BreakPinLinks(*SelfNativePin, true);
+	Schema->BreakPinLinks(*NativePin, true);
 	return true;
 }
 
-bool UBlueprintGraphPin::BreakSinglePinLink(UBlueprintGraphPin* Other)
+bool UBlueprintGraphPinLibrary::BreakSinglePinLink(const FBlueprintGraphPin& Pin, const FBlueprintGraphPin& Other)
 {
-	if(!Other)
-	{
-		return false;
-	}
-
-	UEdGraphPin* SelfNativePin = GetNativePin();
-	UEdGraphPin* OtherNativePin = Other->GetNativePin();
+	UEdGraphPin* SelfNativePin = GetNativePinSafe(Pin);
+	UEdGraphPin* OtherNativePin = GetNativePinSafe(Other);
 	if(!SelfNativePin || !OtherNativePin)
 	{
 		return false;
@@ -120,91 +142,91 @@ bool UBlueprintGraphPin::BreakSinglePinLink(UBlueprintGraphPin* Other)
 	return true;
 }
 
-FName UBlueprintGraphPin::GetPinName() const
+FName UBlueprintGraphPinLibrary::GetPinName(const FBlueprintGraphPin& Pin)
 {
-	UEdGraphPin* Pin = GetNativePin();
-	if(!Pin)
+	UEdGraphPin* NativePin = GetNativePinSafe(Pin);
+	if(!NativePin)
 	{
-		return TEXT("");
+		return NAME_None;
 	}
 
-	return Pin->PinName;
+	return NativePin->PinName;
 }
 
-FString UBlueprintGraphPin::GetPinTypeAsJsonSchema() const
+TEnumAsByte<EEdGraphPinDirection> UBlueprintGraphPinLibrary::GetPinDirection(const FBlueprintGraphPin& Pin)
 {
-	UEdGraphPin* Pin = GetNativePin();
-	if(!Pin)
+	return Pin.PinDirection;
+}
+
+FString UBlueprintGraphPinLibrary::GetPinTypeAsJsonSchema(const FBlueprintGraphPin& Pin)
+{
+	UEdGraphPin* NativePin = GetNativePinSafe(Pin);
+	if(!NativePin)
 	{
 		return TEXT("");
 	}
 
-	UK2Node* OwningNode = Cast<UK2Node>(Pin->GetOwningNode());
+	UK2Node* OwningNode = Cast<UK2Node>(NativePin->GetOwningNode());
 	if(!OwningNode)
 	{
 		return TEXT("");
 	}
 
-	return UBlueprintEditorLibrary::PinTypeToJsonSchema(Pin->PinType, OwningNode->GetBlueprintClassFromNode());
+	return UBlueprintEditorLibrary::PinTypeToJsonSchema(NativePin->PinType, OwningNode->GetBlueprintClassFromNode());
 }
 
-FText UBlueprintGraphPin::GetPinTypeDisplayString() const
+FText UBlueprintGraphPinLibrary::GetPinTypeDisplayString(const FBlueprintGraphPin& Pin)
 {
-	UEdGraphPin* Pin = GetNativePin();
-	if(!Pin)
+	UEdGraphPin* NativePin = GetNativePinSafe(Pin);
+	if(!NativePin)
 	{
 		return FText();
 	}
 
 	const UEdGraphSchema_K2* Schema = GetDefault<UEdGraphSchema_K2>();
-	return Schema->TypeToText(Pin->PinType);
+	return Schema->TypeToText(NativePin->PinType);
 }
 
-FEdGraphPinType UBlueprintGraphPin::GetPinType() const
+FEdGraphPinType UBlueprintGraphPinLibrary::GetPinType(const FBlueprintGraphPin& Pin)
 {
-	UEdGraphPin* Pin = GetNativePin();
-	if(!Pin)
+	UEdGraphPin* NativePin = GetNativePinSafe(Pin);
+	if(!NativePin)
 	{
 		return FEdGraphPinType();
 	}
 
-	return Pin->PinType;
+	return NativePin->PinType;
 }
 
-TEnumAsByte<EEdGraphPinDirection> UBlueprintGraphPin::GetPinDirection() const
+FString UBlueprintGraphPinLibrary::GetPinValue(const FBlueprintGraphPin& Pin)
 {
-	return PinDirection;
-}
-
-FString UBlueprintGraphPin::GetPinValue() const
-{
-	UEdGraphPin* Pin = GetNativePin();
-	if(!Pin)
+	UEdGraphPin* NativePin = GetNativePinSafe(Pin);
+	if(!NativePin)
 	{
 		return TEXT("");
 	}
 
-	return Pin->GetDefaultAsString();
+	return NativePin->GetDefaultAsString();
 }
 
-UK2Node* UBlueprintGraphPin::GetOwningNode() const
+UK2Node* UBlueprintGraphPinLibrary::GetOwningNode(const FBlueprintGraphPin& Pin)
 {
-	return Node;
+	return Pin.Node;
 }
 
-TArray<UBlueprintGraphPin*> UBlueprintGraphPin::ListConnectedPins() const
+TArray<FBlueprintGraphPin> UBlueprintGraphPinLibrary::ListConnectedPins(const FBlueprintGraphPin& Pin)
 {
-	TArray<UBlueprintGraphPin*> Result;
-	UEdGraphPin* Pin = GetNativePin();
-	if(!Pin)
+	TArray<FBlueprintGraphPin> Result;
+	UEdGraphPin* NativePin = GetNativePinSafe(Pin);
+	if(!NativePin)
 	{
 		return Result;
 	}
 
-	for(UEdGraphPin* LinkedToPin : Pin->LinkedTo)
+	for(UEdGraphPin* LinkedToPin : NativePin->LinkedTo)
 	{
-		UBlueprintGraphPin* WrappedPin = UBlueprintGraphPin::FromNativePin(LinkedToPin);
-		if(WrappedPin)
+		FBlueprintGraphPin WrappedPin = FromNativePin(LinkedToPin);
+		if(::IsValid(WrappedPin.Node))
 		{
 			Result.Add(WrappedPin);
 		}
@@ -212,22 +234,47 @@ TArray<UBlueprintGraphPin*> UBlueprintGraphPin::ListConnectedPins() const
 	return Result;
 }
 
-UEdGraphPin* UBlueprintGraphPin::GetNativePin() const
+UEdGraphPin* UBlueprintGraphPinLibrary::GetNativePinSafe(const FBlueprintGraphPin& Pin)
 {
-	if(!::IsValid(Node))
+	if(!::IsValid(Pin.Node))
 	{
 		return nullptr;
 	}
 
-	return Node->FindPin(PinName, PinDirection);
-}
+	// Collect visible pins in the requested direction so we can index into them
+	// and perform an outward search if needed.
+	TArray<UEdGraphPin*, TInlineAllocator<16>> DirectionalPins;
+	for(UEdGraphPin* NativePin : Pin.Node->Pins)
+	{
+		if(!NativePin->bHidden && NativePin->Direction == Pin.PinDirection)
+		{
+			DirectionalPins.Add(NativePin);
+		}
+	}
 
-UEdGraphPin* UBlueprintGraphPin::GetNativePinSafe(const UBlueprintGraphPin* Pin)
-{
-	if(!::IsValid(Pin))
+	// Fast path: index and name both match
+	if(DirectionalPins.IsValidIndex(Pin.PinIndex) && DirectionalPins[Pin.PinIndex]->PinName == Pin.PinName)
 	{
-		return nullptr;
+		return DirectionalPins[Pin.PinIndex];
+	}
+
+	// The pin at the expected index has a different name — the graph has been
+	// modified and indices have drifted. Search outward from PinIndex to find
+	// the nearest pin that still has the correct name.
+	for(int32 Distance = 1; Distance < DirectionalPins.Num(); ++Distance)
+	{
+		const int32 Forward = Pin.PinIndex + Distance;
+		if(DirectionalPins.IsValidIndex(Forward) && DirectionalPins[Forward]->PinName == Pin.PinName)
+		{
+			return DirectionalPins[Forward];
+		}
+
+		const int32 Backward = Pin.PinIndex - Distance;
+		if(DirectionalPins.IsValidIndex(Backward) && DirectionalPins[Backward]->PinName == Pin.PinName)
+		{
+			return DirectionalPins[Backward];
+		}
 	}
 
-	return Pin->GetNativePin();
+	return nullptr;
 }
diff --git a/Engine/Source/Editor/BlueprintEditorLibrary/Private/BlueprintEditorLibrary/BlueprintGraphPin.h b/Engine/Source/Editor/BlueprintEditorLibrary/Private/BlueprintEditorLibrary/BlueprintGraphPin.h
deleted file mode 100644
diff --git a/Engine/Source/Editor/BlueprintEditorLibrary/Public/BlueprintEditorLibrary.h b/Engine/Source/Editor/BlueprintEditorLibrary/Public/BlueprintEditorLibrary.h
--- a/Engine/Source/Editor/BlueprintEditorLibrary/Public/BlueprintEditorLibrary.h
+++ b/Engine/Source/Editor/BlueprintEditorLibrary/Public/BlueprintEditorLibrary.h
@@ -2,6 +2,7 @@
 
 #pragma once
 
+#include "BlueprintEditorLibrary/BlueprintGraphPin.h"
 #include "Components/ActorComponent.h"
 #include "Containers/Array.h"
 #include "Containers/UnrealString.h"
@@ -395,61 +396,73 @@ public:
 	 * Returns all visible pins on this node, optionally discriminated by direction
 	 */
 	UFUNCTION(BlueprintCallable, Category = "Blueprint Graph Editing", meta = (ScriptMethod))
-	static TArray<UBlueprintGraphPin*> ListAllPins(const UK2Node* Node, TEnumAsByte<EEdGraphPinDirection> InDirection = EEdGraphPinDirection::EGPD_MAX);
-	
+	static TArray<FBlueprintGraphPin> ListAllPins(const UK2Node* Node, TEnumAsByte<EEdGraphPinDirection> InDirection = EEdGraphPinDirection::EGPD_MAX);
+
+	/**
+	 * Returns all visible input pins on this node
+	 */
+	UFUNCTION(BlueprintCallable, Category = "Blueprint Graph Editing", meta = (ScriptMethod))
+	static TArray<FBlueprintGraphPin> ListInputPins(const UK2Node* Node);
+
+	/**
+	 * Returns all visible output pins on this node
+	 */
+	UFUNCTION(BlueprintCallable, Category = "Blueprint Graph Editing", meta = (ScriptMethod))
+	static TArray<FBlueprintGraphPin> ListOutputPins(const UK2Node* Node);
+
 	/**
 	 * Returns an input pin specified by index and, optionally, type
 	 */
 	UFUNCTION(BlueprintCallable, Category = "Blueprint Graph Editing", meta = (ScriptMethod))
-	static UBlueprintGraphPin* FindInputPinByIndex(const UK2Node* Node, int32 Index, const FEdGraphPinType& Type = FEdGraphPinType() );
-	
+	static FBlueprintGraphPin FindInputPinByIndex(const UK2Node* Node, int32 Index, const FEdGraphPinType& Type = FEdGraphPinType());
+
 	/**
 	 * Returns an input pin specified by name and, optionally, type
 	 */
 	UFUNCTION(BlueprintCallable, Category = "Blueprint Graph Editing", meta = (ScriptMethod))
-	static UBlueprintGraphPin* FindInputPin(const UK2Node* Node, FName PinName, const FEdGraphPinType& Type = FEdGraphPinType() );
-	
+	static FBlueprintGraphPin FindInputPin(const UK2Node* Node, FName PinName, const FEdGraphPinType& Type = FEdGraphPinType());
+
 	/**
 	 * Returns an output pin specified by index and, optionally, type
 	 */
 	UFUNCTION(BlueprintCallable, Category = "Blueprint Graph Editing", meta = (ScriptMethod))
-	static UBlueprintGraphPin* FindOutputPinByIndex(const UK2Node* Node, int32 Index, const FEdGraphPinType& Type = FEdGraphPinType() );
-	
+	static FBlueprintGraphPin FindOutputPinByIndex(const UK2Node* Node, int32 Index, const FEdGraphPinType& Type = FEdGraphPinType());
+
 	/**
 	 * Returns an output pin specified by name and, optionally, type
 	 */
 	UFUNCTION(BlueprintCallable, Category = "Blueprint Graph Editing", meta = (ScriptMethod))
-	static UBlueprintGraphPin* FindOutputPin(const UK2Node* Node, FName PinName, const FEdGraphPinType& Type = FEdGraphPinType() );
-	
+	static FBlueprintGraphPin FindOutputPin(const UK2Node* Node, FName PinName, const FEdGraphPinType& Type = FEdGraphPinType());
+
 	/**
 	 * Returns the 'execute' or 'do' pin associated with this node, if any
 	 */
 	UFUNCTION(BlueprintCallable, Category = "Blueprint Graph Editing", meta = (ScriptMethod))
-	static UBlueprintGraphPin* FindExecutePin(const UK2Node* Node);
-	
+	static FBlueprintGraphPin FindExecutePin(const UK2Node* Node);
+
 	/**
 	 * Returns the 'then' pin associated with this node, if any
 	 */
 	UFUNCTION(BlueprintCallable, Category = "Blueprint Graph Editing", meta = (ScriptMethod))
-	static UE_API UBlueprintGraphPin* FindThenPin(const UK2Node* Node);
-	
+	static UE_API FBlueprintGraphPin FindThenPin(const UK2Node* Node);
+
 	/**
 	 * Returns the self pin associated with this node, if any
 	 */
 	UFUNCTION(BlueprintCallable, Category = "Blueprint Graph Editing", meta = (ScriptMethod))
-	static UBlueprintGraphPin* FindSelfPin(const UK2Node* Node);
-	
+	static FBlueprintGraphPin FindSelfPin(const UK2Node* Node);
+
 	/**
-	 * Returns the single data output pin associated with this node, returning null if there are multiple or no output(s)
+	 * Returns the single data output pin associated with this node, returning default pin if there are multiple or no output(s)
 	 */
 	UFUNCTION(BlueprintCallable, Category = "Blueprint Graph Editing", meta = (ScriptMethod))
-	static UBlueprintGraphPin* FindResultPin(const UK2Node* Node);
+	static FBlueprintGraphPin FindResultPin(const UK2Node* Node);
 
 	/**
-	 * Returns the single data input pin associated with this node, returning null if there are multiple or no input(s)
+	 * Returns the single data input pin associated with this node, returning default pin if there are multiple or no input(s)
 	 */
 	UFUNCTION(BlueprintCallable, Category = "Blueprint Graph Editing", meta = (ScriptMethod))
-	static UBlueprintGraphPin* FindDataInputPin(const UK2Node* Node);
+	static FBlueprintGraphPin FindDataInputPin(const UK2Node* Node);
 	
 	/**
 	 * Sets this node's visual position in the graph
@@ -479,13 +492,13 @@ public:
 
 	/// begin UK2Node_IfThenElse textension/script methods:
 
-	/** Finds the Else pin on this node, wrapped in a UBlueprintGraphPin */
+	/** Finds the Else pin on this node */
 	UFUNCTION(BlueprintCallable, Category = "Blueprint Graph Editing", meta = (ScriptMethod))
-	static UBlueprintGraphPin* FindElsePin(const UK2Node_IfThenElse* Node);
-	
-	/** Finds the Condition pin on this node, wrapped in a UBlueprintGraphPin */
+	static FBlueprintGraphPin FindElsePin(const UK2Node_IfThenElse* Node);
+
+	/** Finds the Condition pin on this node */
 	UFUNCTION(BlueprintCallable, Category = "Blueprint Graph Editing", meta = (ScriptMethod))
-	static UBlueprintGraphPin* FindConditionPin(const UK2Node_IfThenElse* Node);
+	static FBlueprintGraphPin FindConditionPin(const UK2Node_IfThenElse* Node);
 
 	/// end UK2Node_IfThenElse extension/script methods
 };
diff --git a/Engine/Source/Editor/BlueprintEditorLibrary/Public/BlueprintEditorLibrary/BlueprintGraphPin.h b/Engine/Source/Editor/BlueprintEditorLibrary/Public/BlueprintEditorLibrary/BlueprintGraphPin.h
new file mode 100644
--- /dev/null
+++ b/Engine/Source/Editor/BlueprintEditorLibrary/Public/BlueprintEditorLibrary/BlueprintGraphPin.h
@@ -0,0 +1,151 @@
+// Copyright Epic Games, Inc. All Rights Reserved.
+
+#pragma once
+
+#include "UObject/Object.h"
+#include "EdGraph/EdGraphNode.h"
+
+#include "BlueprintGraphPin.generated.h"
+
+#define UE_API BLUEPRINTEDITORLIBRARY_API
+
+class UEdGraph;
+class UBlueprint;
+
+class UK2Node;
+
+/** 
+	A value type representing a blueprint graph pin (underlying native type: UEdGraphPin). Use IsValid()
+	to detect whether it represents a real pin value (e.g. after finding a pin on a node).
+*/
+USTRUCT(BlueprintType)
+struct FBlueprintGraphPin
+{
+	GENERATED_BODY()
+
+	UPROPERTY()
+	TObjectPtr<UK2Node> Node = nullptr;
+
+	UPROPERTY()
+	TEnumAsByte<EEdGraphPinDirection> PinDirection = EEdGraphPinDirection::EGPD_MAX;
+
+	UPROPERTY()
+	FName PinName;
+
+	UPROPERTY()
+	int32 PinIndex = 0;
+
+	UE_API bool operator==(const FBlueprintGraphPin&) const;
+	UE_API bool operator!=(const FBlueprintGraphPin&) const;
+};
+
+/** 
+ * API for FBlueprintGraphPin
+ */
+UCLASS(Transient, HideDropDown)
+class UBlueprintGraphPinLibrary : public UObject
+{
+	GENERATED_BODY()
+public:
+	/** Creates a FBlueprintGraphPin representing the underlying native Pin */
+	static FBlueprintGraphPin FromNativePin(UEdGraphPin* Pin);
+	
+	/** @return true if the pin has a valid underlying pin */
+	UFUNCTION(BlueprintPure, Category = "Blueprint Editor", meta = (ScriptMethod))
+	static bool IsValid(const FBlueprintGraphPin& Pin);
+
+	/** 
+	 * Because these pin objects are proxies for the underlying native pin, this function
+	 * should be used for establishing whether to FBlueprintGraphPin values refer to the 
+	 * same pin on a node
+	 * 
+	 * @return true if the Other pin and this Pin are the same pin
+	 */
+	UFUNCTION(BlueprintCallable, Category = "Blueprint Editor", meta = (ScriptMethod))
+	static bool IsSameNativePin(const FBlueprintGraphPin& Pin, const FBlueprintGraphPin& Other);
+
+	/**
+	 * Sets the literal value associated with this pin, returning false if the value is not valid
+	 */
+	UFUNCTION(BlueprintCallable, Category = "Blueprint Editor", meta = (ScriptMethod))
+	static bool SetPinValue(const FBlueprintGraphPin& Pin, const FString& Value);
+	
+	/**
+	 * Tries to create a connection between this pin and Other, returning false if the connection is not allowed
+	 */
+	UFUNCTION(BlueprintCallable, Category = "Blueprint Editor", meta = (ScriptMethod))
+	static bool TryCreateConnection(const FBlueprintGraphPin& Pin, const FBlueprintGraphPin& Other);
+
+	/**
+	 * Returns true if this and Other pin can be connected. Performing connection may result in implicitly created conversion nodes.
+	 */
+	UFUNCTION(BlueprintCallable, Category = "Blueprint Editor", meta = (ScriptMethod))
+	static bool CanCreateConnection(const FBlueprintGraphPin& Pin, const FBlueprintGraphPin& Other);
+
+	/**
+	 * Breaks all pin links associated with this pin
+	 */
+	UFUNCTION(BlueprintCallable, Category = "Blueprint Editor", meta = (ScriptMethod))
+	static bool BreakPinLinks(const FBlueprintGraphPin& Pin);
+	
+	/**
+	 * Breaks a pin link between this pin and Other
+	 */
+	UFUNCTION(BlueprintCallable, Category = "Blueprint Editor", meta = (ScriptMethod))
+	static bool BreakSinglePinLink(const FBlueprintGraphPin& Pin, const FBlueprintGraphPin& Other);
+	
+	/**
+	 * Returns the low level name of the pin
+	 */
+	UFUNCTION(BlueprintCallable, Category = "Blueprint Editor", meta = (ScriptMethod))
+	static FName GetPinName(const FBlueprintGraphPin& Pin);
+	
+	/**
+	 * Returns the direction of the pin (either EGPD_Output or EGPD_Input)
+	 */
+	UFUNCTION(BlueprintCallable, Category = "Blueprint Editor", meta = (ScriptMethod))
+	static TEnumAsByte<EEdGraphPinDirection> GetPinDirection(const FBlueprintGraphPin& Pin);
+
+	/**
+	 * Returns a json schema encoded description of the pin's type:
+	 */
+	UFUNCTION(BlueprintCallable, Category = "Blueprint Editor", meta = (ScriptMethod))
+	static FString GetPinTypeAsJsonSchema(const FBlueprintGraphPin& Pin);
+
+	/**
+	 * Returns the string used to describe the pin type in the blueprint UI
+	 */
+	UFUNCTION(BlueprintCallable, Category = "Blueprint Editor", meta = (ScriptMethod))
+	static FText GetPinTypeDisplayString(const FBlueprintGraphPin& Pin);
+
+	/**
+	 * Returns the pin's type, FEdGraphPinType is mostly opaque to script but can
+	 * be used e.g. to create variables of the same type
+	 */
+	UFUNCTION(BlueprintCallable, Category = "Blueprint Editor", meta = (ScriptMethod))
+	static FEdGraphPinType GetPinType(const FBlueprintGraphPin& Pin);
+
+	/**
+	 * Returns the literal value associated with the pin. Empty string for pins 
+	 * that have nothing set or are connected to some other pin:
+	 */
+	UFUNCTION(BlueprintCallable, Category = "Blueprint Editor", meta = (ScriptMethod))
+	static FString GetPinValue(const FBlueprintGraphPin& Pin);
+	
+	/**
+	 * Returns the pin's node
+	 */
+	UFUNCTION(BlueprintCallable, Category = "Blueprint Editor", meta = (ScriptMethod))
+	static UK2Node* GetOwningNode(const FBlueprintGraphPin& Pin);
+
+	/**
+	 * Returns a list of pins this pin is connected to, note that these pin objects are 
+	 * proxies and so comparison to other pin objects must be done via IsSameNativePin
+	 */
+	UFUNCTION(BlueprintCallable, Category = "Blueprint Editor", meta = (ScriptMethod))
+	static TArray<FBlueprintGraphPin> ListConnectedPins(const FBlueprintGraphPin& Pin);
+
+	static UEdGraphPin* GetNativePinSafe(const FBlueprintGraphPin& Pin);
+};
+
+#undef UE_API
