Adding the ability to create UK2Node_ComponentBoundEvent nodes in BlueprintGraphEditor

#Synced-CL 52015895

diff --git a/Engine/Source/Editor/BlueprintEditorLibrary/Private/BlueprintEditorLibrary/BlueprintGraphEditor.cpp b/Engine/Source/Editor/BlueprintEditorLibrary/Private/BlueprintEditorLibrary/BlueprintGraphEditor.cpp
--- a/Engine/Source/Editor/BlueprintEditorLibrary/Private/BlueprintEditorLibrary/BlueprintGraphEditor.cpp
+++ b/Engine/Source/Editor/BlueprintEditorLibrary/Private/BlueprintEditorLibrary/BlueprintGraphEditor.cpp
@@ -9,7 +9,11 @@
 #include "BlueprintVariableNodeSpawner.h"
 #include "Blueprint/BlueprintExceptionInfo.h"
 #include "Kismet/KismetSystemLibrary.h"
+#include "EdGraphSchema_K2_Actions.h"
+#include "Engine/SimpleConstructionScript.h"
+#include "Engine/SCS_Node.h"
 #include "K2Node_CallFunction.h"
+#include "K2Node_ComponentBoundEvent.h"
 #include "K2Node_CustomEvent.h"
 #include "K2Node_Event.h"
 #include "K2Node_FunctionEntry.h"
@@ -21,6 +25,7 @@
 #include "K2Node_VariableGet.h"
 #include "K2Node_VariableSet.h"
 #include "Kismet2/BlueprintEditorUtils.h"
+#include "Kismet2/KismetEditorUtilities.h"
 #include "Kismet2/Kismet2NameValidators.h"
 #include "Logging/StructuredLog.h"
 #include "UObject/Class.h"
@@ -301,6 +306,7 @@ UK2Node_FunctionResult* UBlueprintGraphEditor::AddGraphOutputParameter(FName Out
 	return ResultNodes[0];
 }
 
+
 bool UBlueprintGraphEditor::RemoveGraphOutputParameter(FName OutputName)
 {
 	if(!IsValid())
@@ -1029,6 +1035,89 @@ UK2Node_CustomEvent* UBlueprintGraphEditor::AddCustomEventNode(const FString& Ev
 	return NewNode;
 }
 
+TArray<FName> UBlueprintGraphEditor::ListComponentEvents(UActorComponent* Component) const
+{
+	TArray<FName> Result;
+	if (!Component)
+	{
+		return Result;
+	}
+	for (TFieldIterator<FMulticastDelegateProperty> PropIt(Component->GetClass()); PropIt; ++PropIt)
+	{
+		if (PropIt->HasAnyPropertyFlags(CPF_BlueprintAssignable))
+		{
+			Result.Add(PropIt->GetFName());
+		}
+	}
+	return Result;
+}
+
+UK2Node_ComponentBoundEvent* UBlueprintGraphEditor::AddComponentBoundEventNode(
+	UActorComponent* Component, FName EventName)
+{
+	if (!IsValid() || !Component || EventName == NAME_None)
+	{
+		return nullptr;
+	}
+	if (!FBlueprintEditorUtils::IsEventGraph(Graph))
+	{
+		UE_LOGFMT(LogBlueprint, Warning, "Cannot add component bound event to {Graph} - it is not an event graph", Graph->GetPathName());
+		return nullptr;
+	}
+
+	if (!Blueprint->SkeletonGeneratedClass)
+	{
+		UE_LOGFMT(LogBlueprint, Warning, "Cannot add component bound event - Blueprint {Blueprint} has not been compiled", Blueprint->GetName());
+		return nullptr;
+	}
+
+	// Walk the SCS to find the FObjectProperty corresponding to this component template
+	FObjectProperty* ComponentProperty = nullptr;
+	if (Blueprint->SimpleConstructionScript)
+	{
+		for (USCS_Node* SCSNode : Blueprint->SimpleConstructionScript->GetAllNodes())
+		{
+			if (SCSNode->ComponentTemplate == Component)
+			{
+				ComponentProperty = FindFProperty<FObjectProperty>(Blueprint->SkeletonGeneratedClass, SCSNode->GetVariableName());
+				break;
+			}
+		}
+	}
+	if (!ComponentProperty)
+	{
+		UE_LOGFMT(LogBlueprint, Warning, "Could not find a Blueprint property for component {Component}", Component->GetName());
+		return nullptr;
+	}
+
+	FMulticastDelegateProperty* DelegateProperty = FindFProperty<FMulticastDelegateProperty>(Component->GetClass(), EventName);
+	if (!DelegateProperty)
+	{
+		UE_LOGFMT(LogBlueprint, Warning, "Could not find delegate {EventName} on {ComponentClass}", EventName, Component->GetClass()->GetName());
+		return nullptr;
+	}
+
+	// Return the existing node rather than creating a duplicate
+	if (const UK2Node_ComponentBoundEvent* Existing = FKismetEditorUtilities::FindBoundEventForComponent(Blueprint, EventName, ComponentProperty->GetFName()))
+	{
+		return const_cast<UK2Node_ComponentBoundEvent*>(Existing);
+	}
+
+	UK2Node_ComponentBoundEvent* NewNode = FEdGraphSchemaAction_K2NewNode::SpawnNode<UK2Node_ComponentBoundEvent>(
+		Graph, GetNodeLocation(), EK2NewNodeFlags::SelectNewNode,
+		[ComponentProperty, DelegateProperty](UK2Node_ComponentBoundEvent* Instance)
+		{
+			Instance->InitializeComponentBoundEventParams(ComponentProperty, DelegateProperty);
+		}
+	);
+
+	if (NewNode)
+	{
+		HandleNodeSpawned(NewNode);
+	}
+	return NewNode;
+}
+
 void UBlueprintGraphEditor::RemoveNodes(const TArray<UK2Node*>& Nodes)
 {
 	if(!IsValid())
@@ -1302,7 +1391,6 @@ void UBlueprintGraphEditor::ListAvailableNodes(TArray<FString>& Result, const TA
 	}
 }
 
-
 UEdGraphNode* UBlueprintGraphEditor::CreateNodeFromName(const FString& NodeWithCategory, FVector2D const Location, const TArray<FBlueprintGraphPin>& ContextPins)
 {
 	TRACE_CPUPROFILER_EVENT_SCOPE(UBlueprintGraphEditor::ListAvailableNodes);
diff --git a/Engine/Source/Editor/BlueprintEditorLibrary/Private/BlueprintEditorLibrary/BlueprintGraphEditor.h b/Engine/Source/Editor/BlueprintEditorLibrary/Private/BlueprintEditorLibrary/BlueprintGraphEditor.h
--- a/Engine/Source/Editor/BlueprintEditorLibrary/Private/BlueprintEditorLibrary/BlueprintGraphEditor.h
+++ b/Engine/Source/Editor/BlueprintEditorLibrary/Private/BlueprintEditorLibrary/BlueprintGraphEditor.h
@@ -22,6 +22,8 @@ class UK2Node_Select;
 class UK2Node_FunctionResult;
 class UK2Node_VariableGet;
 class UK2Node_VariableSet;
+class UK2Node_ComponentBoundEvent;
+class UBlueprintGraphPin;
 
 /**
  * Object oriented helper for editing graphs in a blueprint object
@@ -313,6 +315,21 @@ public:
 	UFUNCTION(BlueprintCallable, Category = "Blueprint Editor")
 	UK2Node_CustomEvent* AddCustomEventNode(const FString& EventName);
 
+	/**
+	 * Returns the names of all bindable delegate events on the given component.
+	 */
+	UFUNCTION(BlueprintCallable, Category = "Blueprint Editor")
+	TArray<FName> ListComponentEvents(UActorComponent* Component) const;
+
+	/**
+	 * Creates a UK2Node_ComponentBoundEvent in this event graph for the given component
+	 * and delegate event name. Returns the existing node if one already exists.
+	 * The graph must be an event graph.
+	 */
+	UFUNCTION(BlueprintCallable, Category = "Blueprint Editor")
+	UK2Node_ComponentBoundEvent* AddComponentBoundEventNode(
+		UActorComponent* Component, FName EventName);
+
 	/**
 	 * Removes (aka deletes) the provided nodes from the current graph
 	 */
