Adding get_blueprint_parent_class to BlueprintEditorLibrary

#Synced-CL 52154636

diff --git a/Engine/Source/Editor/BlueprintEditorLibrary/Private/BlueprintEditorLibrary.cpp b/Engine/Source/Editor/BlueprintEditorLibrary/Private/BlueprintEditorLibrary.cpp
--- a/Engine/Source/Editor/BlueprintEditorLibrary/Private/BlueprintEditorLibrary.cpp
+++ b/Engine/Source/Editor/BlueprintEditorLibrary/Private/BlueprintEditorLibrary.cpp
@@ -845,6 +845,11 @@ void UBlueprintEditorLibrary::RefreshAllOpenBlueprintEditors()
 	}
 }
 
+UClass* UBlueprintEditorLibrary::GetBlueprintParentClass(const UBlueprint* Blueprint)
+{
+	return Blueprint ? Blueprint->ParentClass : nullptr;
+}
+
 void UBlueprintEditorLibrary::ReparentBlueprint(UBlueprint* Blueprint, UClass* NewParentClass)
 {
 	if (!Blueprint || !NewParentClass)
diff --git a/Engine/Source/Editor/BlueprintEditorLibrary/Public/BlueprintEditorLibrary.h b/Engine/Source/Editor/BlueprintEditorLibrary/Public/BlueprintEditorLibrary.h
--- a/Engine/Source/Editor/BlueprintEditorLibrary/Public/BlueprintEditorLibrary.h
+++ b/Engine/Source/Editor/BlueprintEditorLibrary/Public/BlueprintEditorLibrary.h
@@ -242,7 +242,16 @@ public:
 	static UE_API void RefreshAllOpenBlueprintEditors();
 
 	/**
-	* Attempts to reparent the given blueprint to the new chosen parent class. 
+	* Returns the parent class of the given Blueprint.
+	*
+	* @param Blueprint			Blueprint to query
+	* @return					The parent class of the Blueprint, or null if the Blueprint is invalid
+	*/
+	UFUNCTION(BlueprintCallable, BlueprintPure, Category = "Blueprint Editor", meta=(ScriptMethod))
+	static UE_API UClass* GetBlueprintParentClass(const UBlueprint* Blueprint);
+
+	/**
+	* Attempts to reparent the given blueprint to the new chosen parent class.
 	*
 	* @param Blueprint			Blueprint that you would like to reparent
 	* @param NewParentClass		The new parent class to use
