Adding the ability to set the Blueprint variable replication settings with Python.

#Synced-CL 52317298

diff --git a/Engine/Source/Editor/BlueprintEditorLibrary/Private/BlueprintEditorLibrary.cpp b/Engine/Source/Editor/BlueprintEditorLibrary/Private/BlueprintEditorLibrary.cpp
--- a/Engine/Source/Editor/BlueprintEditorLibrary/Private/BlueprintEditorLibrary.cpp
+++ b/Engine/Source/Editor/BlueprintEditorLibrary/Private/BlueprintEditorLibrary.cpp
@@ -1245,6 +1245,105 @@ void UBlueprintEditorLibrary::SetBlueprintVariableInstanceEditable(UBlueprint* B
 	FBlueprintEditorUtils::SetBlueprintOnlyEditableFlag(Blueprint, VariableName, !bInstanceEditable);
 }
 
+EBlueprintVariableReplication UBlueprintEditorLibrary::GetBlueprintVariableReplication(
+	UBlueprint* Blueprint, const FName& VariableName)
+{
+	if (!Blueprint)
+	{
+		UE_LOG(LogBlueprintEditorLib, Warning, TEXT("Invalid Blueprint!"));
+		return EBlueprintVariableReplication::None;
+	}
+
+	if (VariableName == NAME_None)
+	{
+		UE_LOG(LogBlueprintEditorLib, Warning, TEXT("Invalid variable name!"));
+		return EBlueprintVariableReplication::None;
+	}
+
+	const uint64* PropFlagPtr = FBlueprintEditorUtils::GetBlueprintVariablePropertyFlags(Blueprint, VariableName);
+	if (!PropFlagPtr)
+	{
+		UE_LOG(LogBlueprintEditorLib, Warning, TEXT("Variable '%s' not found in Blueprint."), *VariableName.ToString());
+		return EBlueprintVariableReplication::None;
+	}
+
+	if (!(*PropFlagPtr & CPF_Net))
+	{
+		return EBlueprintVariableReplication::None;
+	}
+
+	const FName RepNotifyFunc = FBlueprintEditorUtils::GetBlueprintVariableRepNotifyFunc(Blueprint, VariableName);
+	return (RepNotifyFunc != NAME_None)
+		? EBlueprintVariableReplication::RepNotify
+		: EBlueprintVariableReplication::Replicated;
+}
+
+void UBlueprintEditorLibrary::SetBlueprintVariableReplication(UBlueprint* Blueprint, const FName& VariableName,
+	EBlueprintVariableReplication Replication)
+{
+	if (!Blueprint)
+	{
+		UE_LOG(LogBlueprintEditorLib, Warning, TEXT("Invalid Blueprint!"));
+		return;
+	}
+
+	if (VariableName == NAME_None)
+	{
+		UE_LOG(LogBlueprintEditorLib, Warning, TEXT("Invalid variable name!"));
+		return;
+	}
+
+	uint64* PropFlagPtr = FBlueprintEditorUtils::GetBlueprintVariablePropertyFlags(Blueprint, VariableName);
+	if (!PropFlagPtr)
+	{
+		UE_LOG(LogBlueprintEditorLib, Warning, TEXT("Variable '%s' not found in Blueprint."), *VariableName.ToString());
+		return;
+	}
+
+	switch (Replication)
+	{
+	case EBlueprintVariableReplication::None:
+		*PropFlagPtr &= ~CPF_Net;
+		*PropFlagPtr &= ~CPF_RepNotify;
+		FBlueprintEditorUtils::SetBlueprintVariableRepNotifyFunc(Blueprint, VariableName, NAME_None);
+		{
+			const int32 VarIndex = FBlueprintEditorUtils::FindNewVariableIndex(Blueprint, VariableName);
+			if (VarIndex != INDEX_NONE)
+			{
+				Blueprint->NewVariables[VarIndex].ReplicationCondition = COND_None;
+			}
+		}
+		break;
+
+	case EBlueprintVariableReplication::Replicated:
+		*PropFlagPtr |= CPF_Net;
+		*PropFlagPtr &= ~CPF_RepNotify;
+		FBlueprintEditorUtils::SetBlueprintVariableRepNotifyFunc(Blueprint, VariableName, NAME_None);
+		break;
+
+	case EBlueprintVariableReplication::RepNotify:
+		{
+			*PropFlagPtr |= CPF_Net;
+			const FString NewFuncName = FString::Printf(TEXT("OnRep_%s"), *VariableName.ToString());
+			UEdGraph* FuncGraph = FindObject<UEdGraph>(Blueprint, *NewFuncName);
+			if (!FuncGraph)
+			{
+				FuncGraph = FBlueprintEditorUtils::CreateNewGraph(
+					Blueprint, FName(*NewFuncName), UEdGraph::StaticClass(), UEdGraphSchema_K2::StaticClass());
+				FBlueprintEditorUtils::AddFunctionGraph<UClass>(Blueprint, FuncGraph, false, nullptr);
+			}
+			if (FuncGraph)
+			{
+				FBlueprintEditorUtils::SetBlueprintVariableRepNotifyFunc(Blueprint, VariableName, FName(*NewFuncName));
+				*PropFlagPtr |= CPF_RepNotify;
+			}
+		}
+		break;
+	}
+
+	FBlueprintEditorUtils::MarkBlueprintAsStructurallyModified(Blueprint);
+}
+
 bool UBlueprintEditorLibrary::Generic_AddMemberVariableWithValue(UBlueprint* Blueprint, FName MemberName, const uint8* DefaultValuePtr, const FProperty* DefaultValueProp)
 {
 	if(!Blueprint)
diff --git a/Engine/Source/Editor/BlueprintEditorLibrary/Public/BlueprintEditorLibrary.h b/Engine/Source/Editor/BlueprintEditorLibrary/Public/BlueprintEditorLibrary.h
--- a/Engine/Source/Editor/BlueprintEditorLibrary/Public/BlueprintEditorLibrary.h
+++ b/Engine/Source/Editor/BlueprintEditorLibrary/Public/BlueprintEditorLibrary.h
@@ -29,6 +29,21 @@ struct FFrame;
 
 DECLARE_LOG_CATEGORY_EXTERN(LogBlueprintEditorLib, Warning, All);
 
+/**
+ * Network replication mode for a Blueprint variable.
+ * Blueprint-exposed proxy for EVariableReplication (BlueprintDetailsCustomization.h:74).
+ */
+UENUM(BlueprintType)
+enum class EBlueprintVariableReplication : uint8
+{
+	/** Not replicated. */
+	None       UMETA(DisplayName = "None"),
+	/** Replicated from server to client. */
+	Replicated UMETA(DisplayName = "Replicated"),
+	/** Replicated with a notification function called on clients when a new value arrives. */
+	RepNotify  UMETA(DisplayName = "RepNotify"),
+};
+
 /**
  * The results of comparing an assets save version to another
  */
@@ -320,6 +335,29 @@ public:
 	UFUNCTION(BlueprintCallable, Category = "Editor Scripting", meta = (ScriptMethod))
 	static UE_API void SetBlueprintVariableInstanceEditable(UBlueprint* Blueprint, const FName& VariableName, bool bInstanceEditable);
 
+	/**
+	 * Gets the replication mode of a Blueprint variable.
+	 *
+	 * @param Blueprint		The blueprint object
+	 * @param VariableName	The variable name
+	 * @return				The replication mode: None, Replicated, or RepNotify
+	 */
+	UFUNCTION(BlueprintCallable, Category = "Editor Scripting", meta = (ScriptMethod))
+	static UE_API EBlueprintVariableReplication GetBlueprintVariableReplication(
+		UBlueprint* Blueprint, const FName& VariableName);
+
+	/**
+	 * Sets the replication mode on a Blueprint variable.
+	 * RepNotify will auto-create an OnRep_ function if one does not already exist.
+	 *
+	 * @param Blueprint		The blueprint object
+	 * @param VariableName	The variable name
+	 * @param Replication	The replication mode: None, Replicated, or RepNotify
+	 */
+	UFUNCTION(BlueprintCallable, Category = "Editor Scripting", meta = (ScriptMethod))
+	static UE_API void SetBlueprintVariableReplication(UBlueprint* Blueprint, const FName& VariableName,
+		EBlueprintVariableReplication Replication);
+
 	/**
 	 * Creates a blueprint based on a specific parent, honoring registered custom blueprint types
 	 * 
