Adding python exposure for Blueprint comment creation

#Synced-CL 52426408

diff --git a/Engine/Source/Editor/BlueprintEditorLibrary/Private/BlueprintEditorLibrary.cpp b/Engine/Source/Editor/BlueprintEditorLibrary/Private/BlueprintEditorLibrary.cpp
--- a/Engine/Source/Editor/BlueprintEditorLibrary/Private/BlueprintEditorLibrary.cpp
+++ b/Engine/Source/Editor/BlueprintEditorLibrary/Private/BlueprintEditorLibrary.cpp
@@ -1344,6 +1344,59 @@ void UBlueprintEditorLibrary::SetBlueprintVariableReplication(UBlueprint* Bluepr
 	FBlueprintEditorUtils::MarkBlueprintAsStructurallyModified(Blueprint);
 }
 
+FString UBlueprintEditorLibrary::GetCommentText(const UEdGraphNode_Comment* CommentNode)
+{
+	if (!CommentNode) { return FString(); }
+	return CommentNode->NodeComment;
+}
+
+void UBlueprintEditorLibrary::SetCommentText(UEdGraphNode_Comment* CommentNode, const FString& NewText)
+{
+	if (!CommentNode) { return; }
+	CommentNode->Modify();
+	CommentNode->NodeComment = NewText;
+	if (UBlueprint* BP = FBlueprintEditorUtils::FindBlueprintForNode(CommentNode))
+	{
+		FBlueprintEditorUtils::MarkBlueprintAsModified(BP);
+	}
+}
+
+FLinearColor UBlueprintEditorLibrary::GetCommentColor(const UEdGraphNode_Comment* CommentNode)
+{
+	if (!CommentNode) { return FLinearColor::White; }
+	return CommentNode->CommentColor;
+}
+
+void UBlueprintEditorLibrary::SetCommentColor(UEdGraphNode_Comment* CommentNode, FLinearColor Color)
+{
+	if (!CommentNode) { return; }
+	CommentNode->Modify();
+	CommentNode->CommentColor = Color;
+	if (UBlueprint* BP = FBlueprintEditorUtils::FindBlueprintForNode(CommentNode))
+	{
+		FBlueprintEditorUtils::MarkBlueprintAsModified(BP);
+	}
+}
+
+TArray<UK2Node*> UBlueprintEditorLibrary::GetNodesInComment(const UEdGraphNode_Comment* CommentNode)
+{
+	TArray<UK2Node*> Result;
+	if (!CommentNode) { return Result; }
+	for (UObject* Obj : CommentNode->GetNodesUnderComment())
+	{
+		if (UK2Node* K2Node = Cast<UK2Node>(Obj)) { Result.Add(K2Node); }
+	}
+	return Result;
+}
+
+FVector2D UBlueprintEditorLibrary::GetNodeSize(const UEdGraphNode* Node)
+{
+	if (!Node) { return FVector2D::ZeroVector; }
+	const float W = Node->NodeWidth  > 0 ? Node->NodeWidth  : UEdGraphSchema_K2::EstimateNodeWidth(Node);
+	const float H = Node->NodeHeight > 0 ? Node->NodeHeight : UEdGraphSchema_K2::EstimateNodeHeight(const_cast<UEdGraphNode*>(Node));
+	return FVector2D(W, H);
+}
+
 bool UBlueprintEditorLibrary::Generic_AddMemberVariableWithValue(UBlueprint* Blueprint, FName MemberName, const uint8* DefaultValuePtr, const FProperty* DefaultValueProp)
 {
 	if(!Blueprint)
diff --git a/Engine/Source/Editor/BlueprintEditorLibrary/Private/BlueprintEditorLibrary/BlueprintGraphEditor.cpp b/Engine/Source/Editor/BlueprintEditorLibrary/Private/BlueprintEditorLibrary/BlueprintGraphEditor.cpp
--- a/Engine/Source/Editor/BlueprintEditorLibrary/Private/BlueprintEditorLibrary/BlueprintGraphEditor.cpp
+++ b/Engine/Source/Editor/BlueprintEditorLibrary/Private/BlueprintEditorLibrary/BlueprintGraphEditor.cpp
@@ -9,6 +9,8 @@
 #include "BlueprintVariableNodeSpawner.h"
 #include "Blueprint/BlueprintExceptionInfo.h"
 #include "Kismet/KismetSystemLibrary.h"
+#include "EdGraphNode_Comment.h"
+#include "EdGraphSchema_K2.h"
 #include "EdGraphSchema_K2_Actions.h"
 #include "Engine/SimpleConstructionScript.h"
 #include "Engine/SCS_Node.h"
@@ -1118,6 +1120,98 @@ UK2Node_ComponentBoundEvent* UBlueprintGraphEditor::AddComponentBoundEventNode(
 	return NewNode;
 }
 
+UEdGraphNode_Comment* UBlueprintGraphEditor::AddCommentNode(const FString& CommentText, FVector2D Location, FVector2D Size)
+{
+	if (!IsValid())
+	{
+		return nullptr;
+	}
+
+	UEdGraphNode_Comment* CommentNode = NewObject<UEdGraphNode_Comment>(Graph);
+	CommentNode->SetFlags(RF_Transactional);
+	CommentNode->NodePosX = static_cast<int32>(Location.X);
+	CommentNode->NodePosY = static_cast<int32>(Location.Y);
+	CommentNode->NodeWidth = static_cast<int32>(Size.X);
+	CommentNode->NodeHeight = static_cast<int32>(Size.Y);
+
+	CommentNode->PostPlacedNewNode();
+	CommentNode->AllocateDefaultPins();
+
+	// Set text after PostPlacedNewNode, which resets NodeComment to a default value
+	CommentNode->NodeComment = CommentText;
+
+	Graph->Modify();
+	Graph->AddNode(CommentNode, /*bFromUI=*/true, /*bSelectNewNode=*/false);
+
+	FBlueprintEditorUtils::MarkBlueprintAsModified(Blueprint);
+	return CommentNode;
+}
+
+UEdGraphNode_Comment* UBlueprintGraphEditor::AddCommentToNodes(const FString& CommentText, const TArray<UK2Node*>& Nodes, int32 Padding)
+{
+	if (!IsValid() || Nodes.IsEmpty())
+	{
+		return nullptr;
+	}
+
+	// Compute bounding box from node positions
+	int32 MinX = TNumericLimits<int32>::Max();
+	int32 MinY = TNumericLimits<int32>::Max();
+	int32 MaxX = TNumericLimits<int32>::Min();
+	int32 MaxY = TNumericLimits<int32>::Min();
+
+	for (UK2Node* Node : Nodes)
+	{
+		if (!Node) { continue; }
+		const int32 W = Node->NodeWidth  > 0 ? Node->NodeWidth  : static_cast<int32>(UEdGraphSchema_K2::EstimateNodeWidth(Node));
+		const int32 H = Node->NodeHeight > 0 ? Node->NodeHeight : static_cast<int32>(UEdGraphSchema_K2::EstimateNodeHeight(Node));
+		MinX = FMath::Min(MinX, Node->NodePosX);
+		MinY = FMath::Min(MinY, Node->NodePosY);
+		MaxX = FMath::Max(MaxX, Node->NodePosX + W);
+		MaxY = FMath::Max(MaxY, Node->NodePosY + H);
+	}
+
+	if (MinX > MaxX || MinY > MaxY)
+	{
+		return nullptr;
+	}
+
+	const FVector2D Location(MinX - Padding, MinY - Padding);
+	const FVector2D Size(MaxX - MinX + Padding * 2, MaxY - MinY + Padding * 2);
+
+	UEdGraphNode_Comment* CommentNode = AddCommentNode(CommentText, Location, Size);
+	if (CommentNode)
+	{
+		for (UK2Node* Node : Nodes)
+		{
+			if (Node)
+			{
+				CommentNode->AddNodeUnderComment(Node);
+			}
+		}
+	}
+	return CommentNode;
+}
+
+TArray<UEdGraphNode_Comment*> UBlueprintGraphEditor::ListCommentNodes() const
+{
+	TArray<UEdGraphNode_Comment*> Result;
+	if (!IsValid()) { return Result; }
+	for (UEdGraphNode* Node : Graph->Nodes)
+	{
+		if (UEdGraphNode_Comment* Comment = Cast<UEdGraphNode_Comment>(Node)) { Result.Add(Comment); }
+	}
+	return Result;
+}
+
+void UBlueprintGraphEditor::RemoveCommentNode(UEdGraphNode_Comment* CommentNode)
+{
+	if (!IsValid() || !CommentNode) { return; }
+	Graph->Modify();
+	Graph->RemoveNode(CommentNode);
+	FBlueprintEditorUtils::MarkBlueprintAsModified(Blueprint);
+}
+
 void UBlueprintGraphEditor::RemoveNodes(const TArray<UK2Node*>& Nodes)
 {
 	if(!IsValid())
diff --git a/Engine/Source/Editor/BlueprintEditorLibrary/Private/BlueprintEditorLibrary/BlueprintGraphEditor.h b/Engine/Source/Editor/BlueprintEditorLibrary/Private/BlueprintEditorLibrary/BlueprintGraphEditor.h
--- a/Engine/Source/Editor/BlueprintEditorLibrary/Private/BlueprintEditorLibrary/BlueprintGraphEditor.h
+++ b/Engine/Source/Editor/BlueprintEditorLibrary/Private/BlueprintEditorLibrary/BlueprintGraphEditor.h
@@ -23,6 +23,7 @@ class UK2Node_FunctionResult;
 class UK2Node_VariableGet;
 class UK2Node_VariableSet;
 class UK2Node_ComponentBoundEvent;
+class UEdGraphNode_Comment;
 class UBlueprintGraphPin;
 
 /**
@@ -330,6 +331,31 @@ public:
 	UK2Node_ComponentBoundEvent* AddComponentBoundEventNode(
 		UActorComponent* Component, FName EventName);
 
+	/**
+	 * Creates a comment node in the current graph at the given location with the given text and size.
+	 */
+	UFUNCTION(BlueprintCallable, Category = "Blueprint Editor")
+	UEdGraphNode_Comment* AddCommentNode(const FString& CommentText, FVector2D Location, FVector2D Size = FVector2D(400, 200));
+
+	/**
+	 * Creates a comment node that wraps the given nodes with the given padding.
+	 * The comment bounds are computed from the node positions plus padding.
+	 */
+	UFUNCTION(BlueprintCallable, Category = "Blueprint Editor")
+	UEdGraphNode_Comment* AddCommentToNodes(const FString& CommentText, const TArray<UK2Node*>& Nodes, int32 Padding = 50);
+
+	/**
+	 * Returns all comment nodes in the current graph.
+	 */
+	UFUNCTION(BlueprintCallable, Category = "Blueprint Editor")
+	TArray<UEdGraphNode_Comment*> ListCommentNodes() const;
+
+	/**
+	 * Removes a comment node from the current graph.
+	 */
+	UFUNCTION(BlueprintCallable, Category = "Blueprint Editor")
+	void RemoveCommentNode(UEdGraphNode_Comment* CommentNode);
+
 	/**
 	 * Removes (aka deletes) the provided nodes from the current graph
 	 */
diff --git a/Engine/Source/Editor/BlueprintEditorLibrary/Public/BlueprintEditorLibrary.h b/Engine/Source/Editor/BlueprintEditorLibrary/Public/BlueprintEditorLibrary.h
--- a/Engine/Source/Editor/BlueprintEditorLibrary/Public/BlueprintEditorLibrary.h
+++ b/Engine/Source/Editor/BlueprintEditorLibrary/Public/BlueprintEditorLibrary.h
@@ -24,6 +24,9 @@ class FProperty;
 class UBlueprint;
 class UClass;
 class UEdGraph;
+class UEdGraphNode;
+class UEdGraphNode_Comment;
+class UK2Node;
 class UObject;
 struct FFrame;
 
@@ -358,9 +361,51 @@ public:
 	static UE_API void SetBlueprintVariableReplication(UBlueprint* Blueprint, const FName& VariableName,
 		EBlueprintVariableReplication Replication);
 
+	/**
+	 * Returns the comment text of a comment node.
+	 */
+	UFUNCTION(BlueprintCallable, Category = "Blueprint Editor", meta = (ScriptMethod))
+	static UE_API FString GetCommentText(const UEdGraphNode_Comment* CommentNode);
+
+	/**
+	 * Sets the comment text of a comment node.
+	 */
+	UFUNCTION(BlueprintCallable, Category = "Blueprint Editor", meta = (ScriptMethod))
+	static UE_API void SetCommentText(UEdGraphNode_Comment* CommentNode, const FString& NewText);
+
+	/**
+	 * Gets the background color of a comment node.
+	 */
+	UFUNCTION(BlueprintCallable, Category = "Blueprint Editor", meta = (ScriptMethod))
+	static UE_API FLinearColor GetCommentColor(const UEdGraphNode_Comment* CommentNode);
+
+	/**
+	 * Sets the background color of a comment node.
+	 */
+	UFUNCTION(BlueprintCallable, Category = "Blueprint Editor", meta = (ScriptMethod))
+	static UE_API void SetCommentColor(UEdGraphNode_Comment* CommentNode, FLinearColor Color);
+
+	/**
+	 * Returns the K2 nodes contained within the given comment node.
+	 * Only UK2Node-derived nodes are returned; nested comments and other non-K2 node
+	 * types are excluded.
+	 * Membership is explicitly maintained when comments are created programmatically.
+	 * Note: if the comment is manually moved or resized in the editor, membership is
+	 * rebuilt spatially and may diverge from the programmatic assignment.
+	 */
+	UFUNCTION(BlueprintCallable, Category = "Blueprint Editor", meta = (ScriptMethod))
+	static UE_API TArray<UK2Node*> GetNodesInComment(const UEdGraphNode_Comment* CommentNode);
+
+	/**
+	 * Returns the size (width, height) of a node as stored in the graph data.
+	 * Note: sizes are only populated after the graph has been rendered in the editor.
+	 */
+	UFUNCTION(BlueprintCallable, Category = "Blueprint Editor", meta = (ScriptMethod))
+	static UE_API FVector2D GetNodeSize(const UEdGraphNode* Node);
+
 	/**
 	 * Creates a blueprint based on a specific parent, honoring registered custom blueprint types
-	 * 
+	 *
 	 * @param AssetPath				The full path that the asset should be created with
 	 * @param ParentClass			The parent class that the blueprint should be based on
 	 */
