Exposing EventDispatcher creation and usage for python

#Synced-CL 52518884

diff --git a/Engine/Source/Editor/BlueprintEditorLibrary/Private/BlueprintEditorLibrary.cpp b/Engine/Source/Editor/BlueprintEditorLibrary/Private/BlueprintEditorLibrary.cpp
--- a/Engine/Source/Editor/BlueprintEditorLibrary/Private/BlueprintEditorLibrary.cpp
+++ b/Engine/Source/Editor/BlueprintEditorLibrary/Private/BlueprintEditorLibrary.cpp
@@ -29,6 +29,7 @@
 #include "JsonSchema/JsonSchemaGenerator.h"
 #include "Internationalization/Text.h"
 #include "K2Node_CallFunction.h"
+#include "K2Node_FunctionEntry.h"
 #include "K2Node_CommutativeAssociativeBinaryOperator.h"
 #include "K2Node_IfThenElse.h"
 #include "K2Node_PromotableOperator.h"
@@ -1461,6 +1462,202 @@ bool UBlueprintEditorLibrary::AddMemberVariable(UBlueprint* Blueprint, FName Mem
 		VariableType);
 }
 
+bool UBlueprintEditorLibrary::AddEventDispatcher(UBlueprint* Blueprint, FName Name)
+{
+	if (!::IsValid(Blueprint) || Name == NAME_None)
+	{
+		UE_LOG(LogBlueprintEditorLib, Warning, TEXT("AddEventDispatcher: invalid Blueprint or empty name"));
+		return false;
+	}
+
+	if (FBlueprintEditorUtils::FindNewVariableIndex(Blueprint, Name) != INDEX_NONE)
+	{
+		UE_LOG(LogBlueprintEditorLib, Warning,
+			TEXT("AddEventDispatcher: name '%s' is already in use on '%s'"),
+			*Name.ToString(), *Blueprint->GetPathName());
+		return false;
+	}
+
+	Blueprint->Modify();
+
+	FEdGraphPinType DelegateType;
+	DelegateType.PinCategory = UEdGraphSchema_K2::PC_MCDelegate;
+	if (!FBlueprintEditorUtils::AddMemberVariable(Blueprint, Name, DelegateType))
+	{
+		UE_LOG(LogBlueprintEditorLib, Warning,
+			TEXT("AddEventDispatcher: AddMemberVariable failed for '%s' on '%s'"),
+			*Name.ToString(), *Blueprint->GetPathName());
+		return false;
+	}
+
+	const UEdGraphSchema_K2* K2Schema = GetDefault<UEdGraphSchema_K2>();
+	UEdGraph* SignatureGraph = FBlueprintEditorUtils::CreateNewGraph(Blueprint, Name, UEdGraph::StaticClass(), UEdGraphSchema_K2::StaticClass());
+	if (!SignatureGraph)
+	{
+		FBlueprintEditorUtils::RemoveMemberVariable(Blueprint, Name);
+		UE_LOG(LogBlueprintEditorLib, Warning,
+			TEXT("AddEventDispatcher: failed to create signature graph for '%s'"), *Name.ToString());
+		return false;
+	}
+
+	SignatureGraph->bEditable = false;
+	K2Schema->CreateDefaultNodesForGraph(*SignatureGraph);
+	K2Schema->CreateFunctionGraphTerminators(*SignatureGraph, (UClass*)nullptr);
+	K2Schema->AddExtraFunctionFlags(SignatureGraph, FUNC_BlueprintCallable | FUNC_BlueprintEvent | FUNC_Public);
+	K2Schema->MarkFunctionEntryAsEditable(SignatureGraph, true);
+
+	Blueprint->DelegateSignatureGraphs.Add(SignatureGraph);
+	FBlueprintEditorUtils::MarkBlueprintAsStructurallyModified(Blueprint);
+
+	return true;
+}
+
+bool UBlueprintEditorLibrary::RemoveEventDispatcher(UBlueprint* Blueprint, FName Name)
+{
+	if (!::IsValid(Blueprint) || Name == NAME_None)
+	{
+		UE_LOG(LogBlueprintEditorLib, Warning, TEXT("RemoveEventDispatcher: invalid Blueprint or empty name"));
+		return false;
+	}
+
+	const int32 VarIndex = FBlueprintEditorUtils::FindNewVariableIndex(Blueprint, Name);
+	if (VarIndex == INDEX_NONE)
+	{
+		UE_LOG(LogBlueprintEditorLib, Warning,
+			TEXT("RemoveEventDispatcher: no dispatcher named '%s' on '%s'"),
+			*Name.ToString(), *Blueprint->GetPathName());
+		return false;
+	}
+
+	if (Blueprint->NewVariables[VarIndex].VarType.PinCategory != UEdGraphSchema_K2::PC_MCDelegate)
+	{
+		UE_LOG(LogBlueprintEditorLib, Warning,
+			TEXT("RemoveEventDispatcher: '%s' is not an event dispatcher"), *Name.ToString());
+		return false;
+	}
+
+	Blueprint->Modify();
+
+	// Remove the variable entry and any nodes referencing it in all graphs
+	FBlueprintEditorUtils::RemoveMemberVariable(Blueprint, Name);
+
+	// RemoveMemberVariable does not clean up the signature graph — do it explicitly
+	UEdGraph* SignatureGraph = FBlueprintEditorUtils::GetDelegateSignatureGraphByName(Blueprint, Name);
+	if (SignatureGraph)
+	{
+		Blueprint->DelegateSignatureGraphs.Remove(SignatureGraph);
+		FBlueprintEditorUtils::RemoveGraph(Blueprint, SignatureGraph);
+	}
+
+	return true;
+}
+
+TArray<FName> UBlueprintEditorLibrary::ListEventDispatchers(const UBlueprint* Blueprint)
+{
+	TArray<FName> Result;
+	if (!::IsValid(Blueprint))
+	{
+		return Result;
+	}
+
+	for (const FBPVariableDescription& Var : Blueprint->NewVariables)
+	{
+		if (Var.VarType.PinCategory == UEdGraphSchema_K2::PC_MCDelegate)
+		{
+			Result.Add(Var.VarName);
+		}
+	}
+	return Result;
+}
+
+bool UBlueprintEditorLibrary::AddEventDispatcherParameter(UBlueprint* Blueprint, FName DispatcherName,
+	FName ParamName, const FEdGraphPinType& ParamType)
+{
+	if (!::IsValid(Blueprint) || DispatcherName == NAME_None || ParamName == NAME_None)
+	{
+		UE_LOG(LogBlueprintEditorLib, Warning, TEXT("AddEventDispatcherParameter: invalid arguments"));
+		return false;
+	}
+
+	UEdGraph* SignatureGraph = FBlueprintEditorUtils::GetDelegateSignatureGraphByName(Blueprint, DispatcherName);
+	if (!SignatureGraph)
+	{
+		UE_LOG(LogBlueprintEditorLib, Warning,
+			TEXT("AddEventDispatcherParameter: no dispatcher named '%s' on '%s'"),
+			*DispatcherName.ToString(), *Blueprint->GetPathName());
+		return false;
+	}
+
+	TArray<UK2Node_FunctionEntry*> EntryNodes;
+	SignatureGraph->GetNodesOfClass(EntryNodes);
+	if (EntryNodes.Num() == 0)
+	{
+		UE_LOG(LogBlueprintEditorLib, Warning,
+			TEXT("AddEventDispatcherParameter: signature graph for '%s' has no entry node"),
+			*DispatcherName.ToString());
+		return false;
+	}
+
+	UK2Node_FunctionEntry* EntryNode = EntryNodes[0];
+	if (EntryNode->FindPin(ParamName, EGPD_Output))
+	{
+		UE_LOG(LogBlueprintEditorLib, Warning,
+			TEXT("AddEventDispatcherParameter: parameter '%s' already exists on '%s'"),
+			*ParamName.ToString(), *DispatcherName.ToString());
+		return false;
+	}
+
+	EntryNode->Modify();
+	EntryNode->CreateUserDefinedPin(ParamName, ParamType, EGPD_Output);
+	FBlueprintEditorUtils::MarkBlueprintAsStructurallyModified(Blueprint);
+
+	return true;
+}
+
+bool UBlueprintEditorLibrary::RemoveEventDispatcherParameter(UBlueprint* Blueprint, FName DispatcherName,
+	FName ParamName)
+{
+	if (!::IsValid(Blueprint) || DispatcherName == NAME_None || ParamName == NAME_None)
+	{
+		UE_LOG(LogBlueprintEditorLib, Warning, TEXT("RemoveEventDispatcherParameter: invalid arguments"));
+		return false;
+	}
+
+	UEdGraph* SignatureGraph = FBlueprintEditorUtils::GetDelegateSignatureGraphByName(Blueprint, DispatcherName);
+	if (!SignatureGraph)
+	{
+		UE_LOG(LogBlueprintEditorLib, Warning,
+			TEXT("RemoveEventDispatcherParameter: no dispatcher named '%s' on '%s'"),
+			*DispatcherName.ToString(), *Blueprint->GetPathName());
+		return false;
+	}
+
+	TArray<UK2Node_FunctionEntry*> EntryNodes;
+	SignatureGraph->GetNodesOfClass(EntryNodes);
+	if (EntryNodes.Num() == 0)
+	{
+		UE_LOG(LogBlueprintEditorLib, Warning,
+			TEXT("RemoveEventDispatcherParameter: signature graph for '%s' has no entry node"),
+			*DispatcherName.ToString());
+		return false;
+	}
+
+	UK2Node_FunctionEntry* EntryNode = EntryNodes[0];
+	if (!EntryNode->FindPin(ParamName, EGPD_Output))
+	{
+		UE_LOG(LogBlueprintEditorLib, Warning,
+			TEXT("RemoveEventDispatcherParameter: no parameter '%s' on dispatcher '%s'"),
+			*ParamName.ToString(), *DispatcherName.ToString());
+		return false;
+	}
+
+	EntryNode->Modify();
+	EntryNode->RemoveUserDefinedPinByName(ParamName);
+	FBlueprintEditorUtils::MarkBlueprintAsStructurallyModified(Blueprint);
+
+	return true;
+}
+
 FEdGraphPinType UBlueprintEditorLibrary::GetBasicTypeByName(FName TypeName)
 {
 	FEdGraphPinType Result;
diff --git a/Engine/Source/Editor/BlueprintEditorLibrary/Private/BlueprintEditorLibrary/BlueprintGraphEditor.cpp b/Engine/Source/Editor/BlueprintEditorLibrary/Private/BlueprintEditorLibrary/BlueprintGraphEditor.cpp
--- a/Engine/Source/Editor/BlueprintEditorLibrary/Private/BlueprintEditorLibrary/BlueprintGraphEditor.cpp
+++ b/Engine/Source/Editor/BlueprintEditorLibrary/Private/BlueprintEditorLibrary/BlueprintGraphEditor.cpp
@@ -1037,6 +1037,54 @@ UK2Node_CustomEvent* UBlueprintGraphEditor::AddCustomEventNode(const FString& Ev
 	return NewNode;
 }
 
+UK2Node_Event* UBlueprintGraphEditor::AddDispatcherEventNode(FName DispatcherName)
+{
+	if (!IsValid())
+	{
+		return nullptr;
+	}
+
+	if (!FBlueprintEditorUtils::IsEventGraph(Graph))
+	{
+		UE_LOGFMT(LogBlueprint, Warning,
+			"Cannot add dispatcher event node to {Graph} - it is not an event graph",
+			Graph->GetPathName());
+		return nullptr;
+	}
+
+	if (!Blueprint->SkeletonGeneratedClass)
+	{
+		UE_LOGFMT(LogBlueprint, Warning,
+			"Cannot add dispatcher event node - {Blueprint} has no skeleton class",
+			Blueprint->GetName());
+		return nullptr;
+	}
+
+	FMulticastDelegateProperty* DelegateProperty = FindFProperty<FMulticastDelegateProperty>(
+		Blueprint->SkeletonGeneratedClass, DispatcherName);
+	if (!DelegateProperty || !DelegateProperty->SignatureFunction)
+	{
+		UE_LOGFMT(LogBlueprint, Warning,
+			"AddDispatcherEventNode: no dispatcher '{Name}' found on skeleton class",
+			DispatcherName);
+		return nullptr;
+	}
+	const UFunction* SignatureFunction = DelegateProperty->SignatureFunction;
+
+	// The event name must be unique and must not match the dispatcher name — having both an event
+	// stub function and a delegate variable with the same name causes a compile conflict.
+	FName EventName = FBlueprintEditorUtils::FindUniqueKismetName(Blueprint, DispatcherName.ToString());
+	UK2Node_CustomEvent* NewNode = UK2Node_CustomEvent::CreateFromFunction(
+		GetNodeLocation(), Graph, EventName.ToString(), SignatureFunction);
+	if (!NewNode)
+	{
+		return nullptr;
+	}
+
+	HandleNodeSpawned(NewNode);
+	return NewNode;
+}
+
 TArray<FName> UBlueprintGraphEditor::ListComponentEvents(UActorComponent* Component) const
 {
 	TArray<FName> Result;
diff --git a/Engine/Source/Editor/BlueprintEditorLibrary/Private/BlueprintEditorLibrary/BlueprintGraphEditor.h b/Engine/Source/Editor/BlueprintEditorLibrary/Private/BlueprintEditorLibrary/BlueprintGraphEditor.h
--- a/Engine/Source/Editor/BlueprintEditorLibrary/Private/BlueprintEditorLibrary/BlueprintGraphEditor.h
+++ b/Engine/Source/Editor/BlueprintEditorLibrary/Private/BlueprintEditorLibrary/BlueprintGraphEditor.h
@@ -316,6 +316,15 @@ public:
 	UFUNCTION(BlueprintCallable, Category = "Blueprint Editor")
 	UK2Node_CustomEvent* AddCustomEventNode(const FString& EventName);
 
+	/**
+	 * Creates an event node in this event graph whose pins match the signature of the named dispatcher.
+	 * The current graph must be an Event Graph and the Blueprint must have a valid SkeletonGeneratedClass
+	 * (i.e. it must have been compiled at least once). Returns nullptr if the graph is not an Event Graph,
+	 * the skeleton class is absent, or no dispatcher named DispatcherName exists on the Blueprint.
+	 */
+	UFUNCTION(BlueprintCallable, Category = "Blueprint Editor")
+	UK2Node_Event* AddDispatcherEventNode(FName DispatcherName);
+
 	/**
 	 * Returns the names of all bindable delegate events on the given component.
 	 */
diff --git a/Engine/Source/Editor/BlueprintEditorLibrary/Public/BlueprintEditorLibrary.h b/Engine/Source/Editor/BlueprintEditorLibrary/Public/BlueprintEditorLibrary.h
--- a/Engine/Source/Editor/BlueprintEditorLibrary/Public/BlueprintEditorLibrary.h
+++ b/Engine/Source/Editor/BlueprintEditorLibrary/Public/BlueprintEditorLibrary.h
@@ -442,11 +442,47 @@ public:
 	
 	/**
 	  * Adds a member variable to the specified blueprint with the specified type.
-	  * 
+	  *
 	  * @return	true if it succeeds, false if it fails.
 	  */
 	UFUNCTION(BlueprintCallable, Category = "Blueprint Editor")
 	static UE_API bool AddMemberVariable(UBlueprint* Blueprint, FName MemberName, const FEdGraphPinType& VariableType);
+
+	/**
+	 * Creates a new event dispatcher on the Blueprint with the given name.
+	 * Returns false if the name is already in use or the Blueprint is invalid.
+	 */
+	UFUNCTION(BlueprintCallable, Category = "Blueprint Editor", meta = (ScriptMethod))
+	static UE_API bool AddEventDispatcher(UBlueprint* Blueprint, FName Name);
+
+	/**
+	 * Removes the event dispatcher with the given name from the Blueprint.
+	 * Returns false if no dispatcher with that name exists or the Blueprint is invalid.
+	 */
+	UFUNCTION(BlueprintCallable, Category = "Blueprint Editor", meta = (ScriptMethod))
+	static UE_API bool RemoveEventDispatcher(UBlueprint* Blueprint, FName Name);
+
+	/**
+	 * Returns the names of all event dispatchers defined on the Blueprint.
+	 */
+	UFUNCTION(BlueprintCallable, Category = "Blueprint Editor", meta = (ScriptMethod))
+	static UE_API TArray<FName> ListEventDispatchers(const UBlueprint* Blueprint);
+
+	/**
+	 * Adds a parameter to an event dispatcher's signature.
+	 * Returns false if the dispatcher does not exist or the parameter name is already in use.
+	 */
+	UFUNCTION(BlueprintCallable, Category = "Blueprint Editor", meta = (ScriptMethod))
+	static UE_API bool AddEventDispatcherParameter(UBlueprint* Blueprint, FName DispatcherName,
+		FName ParamName, const FEdGraphPinType& ParamType);
+
+	/**
+	 * Removes a parameter from an event dispatcher's signature.
+	 * Returns false if the dispatcher or parameter does not exist.
+	 */
+	UFUNCTION(BlueprintCallable, Category = "Blueprint Editor", meta = (ScriptMethod))
+	static UE_API bool RemoveEventDispatcherParameter(UBlueprint* Blueprint, FName DispatcherName,
+		FName ParamName);
 	
 	/** @return a pintype for 'int', 'byte', 'bool', 'real', 'name', 'string' or 'text' - returns 'int' type if invalid type is provided */
 	UFUNCTION(BlueprintCallable, Category = "Blueprint Editor")
